Rename ConcertServerUI to MultiUserServer and move it to ConcertApp folder

diff --git a/Engine/Plugins/Developer/Concert/ConcertApp/MultiUserServer/MultiUserServer.uplugin b/Engine/Plugins/Developer/Concert/ConcertApp/MultiUserServer/MultiUserServer.uplugin
new file mode 100644
--- /dev/null
+++ b/Engine/Plugins/Developer/Concert/ConcertApp/MultiUserServer/MultiUserServer.uplugin
@@ -0,0 +1,45 @@
+{
+	"FileVersion": 3,
+	"Version": 1,
+	"VersionName": "0.2",
+	"FriendlyName": "Multi User Server",
+	"Description": "Visualizes the multi-user server",
+	"Category": "Networking",
+	"CreatedBy": "Epic Games, Inc.",
+	"CreatedByURL": "http://epicgames.com",
+	"DocsURL": "",
+	"MarketplaceURL": "",
+	"SupportURL": "",
+	"EnabledByDefault": false,
+	"CanContainContent": false,
+	"Hidden": true,
+	"IsBetaVersion": true,
+	"Installed": false,
+	"SupportedPrograms": [ "UnrealMultiUserServer", "UnrealRecoverySvc", "CrashReportClientEditor" ],
+	"Modules": [
+        {
+          "Name": "MultiUserServer",
+          "Type": "UncookedOnly",
+          "LoadingPhase": "Default",
+          "ProgramAllowList": [ "UnrealMultiUserServer", "UnrealRecoverySvc", "CrashReportClientEditor" ]
+        }
+	],
+	"Plugins": [
+		{
+			"Name": "ConcertMain",
+			"Enabled": true
+		},
+        {
+          "Name": "ConcertSharedSlate",
+          "Enabled": true
+        },
+		{
+			"Name": "ConcertSyncCore",
+			"Enabled": true
+		},
+		{
+			"Name": "ConcertSyncServer",
+			"Enabled": true
+		}
+	]
+}
diff --git a/Engine/Plugins/Developer/Concert/ConcertApp/MultiUserServer/Source/MultiUserServer/MultiUserServer.build.cs b/Engine/Plugins/Developer/Concert/ConcertApp/MultiUserServer/Source/MultiUserServer/MultiUserServer.build.cs
new file mode 100644
--- /dev/null
+++ b/Engine/Plugins/Developer/Concert/ConcertApp/MultiUserServer/Source/MultiUserServer/MultiUserServer.build.cs
@@ -0,0 +1,44 @@
+// Copyright Epic Games, Inc. All Rights Reserved.
+
+namespace UnrealBuildTool.Rules
+{
+	public class MultiUserServer : ModuleRules
+	{
+		public MultiUserServer(ReadOnlyTargetRules Target) : base(Target)
+		{
+			PrivateDependencyModuleNames.AddRange(
+				new string[]
+				{
+					"ApplicationCore",
+					"Core",
+					"CoreUObject",
+					"Concert",
+					"ConcertSharedSlate",
+					"ConcertSyncCore",
+					"ConcertTransport",
+					"EditorStyle",			// Needed so OutputLog module to work correctly
+					"InputCore",
+					"OutputLog",
+					"Projects",
+					"Slate",
+					"SlateCore",
+					"StandaloneRenderer",	
+					"ToolWidgets",
+					"WorkspaceMenuStructure" // Needed so OutputLog module to work correctly
+				}
+			);
+			
+			PrivateIncludePathModuleNames.AddRange(
+				new string[] {
+					"ConcertSyncServer",
+				}
+			);
+			
+			DynamicallyLoadedModuleNames.AddRange(
+				new string[]
+				{
+					"ConcertSyncServer"
+				});
+		}
+	}
+}
diff --git a/Engine/Plugins/Developer/Concert/ConcertApp/MultiUserServer/Source/MultiUserServer/Private/ConcertConsoleCommandExecutor.cpp b/Engine/Plugins/Developer/Concert/ConcertApp/MultiUserServer/Source/MultiUserServer/Private/ConcertConsoleCommandExecutor.cpp
new file mode 100644
--- /dev/null
+++ b/Engine/Plugins/Developer/Concert/ConcertApp/MultiUserServer/Source/MultiUserServer/Private/ConcertConsoleCommandExecutor.cpp
@@ -0,0 +1,88 @@
+﻿// Copyright Epic Games, Inc. All Rights Reserved.
+
+#include "ConcertConsoleCommandExecutor.h"
+
+#include "CoreGlobals.h"
+#include "Framework/Commands/InputChord.h"
+#include "Misc/OutputDevice.h"
+
+#define LOCTEXT_NAMESPACE "UnrealMultiUserUI"
+
+FName FConcertConsoleCommandExecutor::StaticName()
+{
+	static const FName CmdExecName = TEXT("Cmd");
+	return CmdExecName;
+}
+
+FName FConcertConsoleCommandExecutor::GetName() const
+{
+	return StaticName();
+}
+
+FText FConcertConsoleCommandExecutor::GetDisplayName() const
+{
+	return LOCTEXT("ConsoleCommandExecutorDisplayName", "Cmd");
+}
+
+FText FConcertConsoleCommandExecutor::GetDescription() const
+{
+	return LOCTEXT("ConsoleCommandExecutorDescription", "Execute Unreal Console Commands");
+}
+
+FText FConcertConsoleCommandExecutor::GetHintText() const
+{
+	return LOCTEXT("ConsoleCommandExecutorHintText", "Enter Console Command");
+}
+
+void FConcertConsoleCommandExecutor::GetAutoCompleteSuggestions(const TCHAR* Input, TArray<FString>& Out)
+{
+	auto OnConsoleVariable = [&Out](const TCHAR *Name, IConsoleObject* CVar)
+	{
+		if (CVar->TestFlags(ECVF_Unregistered))
+		{
+			return;
+		}
+
+		Out.Add(Name);
+	};
+
+	IConsoleManager::Get().ForEachConsoleObjectThatContains(FConsoleObjectVisitor::CreateLambda(OnConsoleVariable), Input);
+}
+
+void FConcertConsoleCommandExecutor::GetExecHistory(TArray<FString>& Out)
+{
+	IConsoleManager::Get().GetConsoleHistory(TEXT(""), Out);
+}
+
+bool FConcertConsoleCommandExecutor::Exec(const TCHAR* Input)
+{
+	IConsoleManager::Get().AddConsoleHistoryEntry(TEXT(""), Input);
+	
+	FOutputDevice& Output = *GLog;
+	UWorld* World = nullptr;
+	
+	bool bSuccess = IConsoleManager::Get().ProcessUserConsoleInput(Input, Output, World);
+	if (!bSuccess)
+	{
+		bSuccess |= FSelfRegisteringExec::StaticExec(World, Input, Output);
+	}
+
+	return bSuccess;
+}
+
+bool FConcertConsoleCommandExecutor::AllowHotKeyClose() const
+{
+	return true;
+}
+
+bool FConcertConsoleCommandExecutor::AllowMultiLine() const
+{
+	return false;
+}
+
+FInputChord FConcertConsoleCommandExecutor::GetHotKey() const
+{
+	return FInputChord();
+}
+
+#undef LOCTEXT_NAMESPACE
diff --git a/Engine/Plugins/Developer/Concert/ConcertApp/MultiUserServer/Source/MultiUserServer/Private/ConcertConsoleCommandExecutor.h b/Engine/Plugins/Developer/Concert/ConcertApp/MultiUserServer/Source/MultiUserServer/Private/ConcertConsoleCommandExecutor.h
new file mode 100644
--- /dev/null
+++ b/Engine/Plugins/Developer/Concert/ConcertApp/MultiUserServer/Source/MultiUserServer/Private/ConcertConsoleCommandExecutor.h
@@ -0,0 +1,27 @@
+﻿// Copyright Epic Games, Inc. All Rights Reserved.
+
+#pragma once
+
+#include "CoreMinimal.h"
+#include "HAL/IConsoleManager.h"
+
+/** Executes commands put into the output log */
+class FConcertConsoleCommandExecutor : public IConsoleCommandExecutor
+{
+public:
+	
+	static FName StaticName();
+
+	//~ Begin IConsoleCommandExecutor Interface
+	virtual FName GetName() const override;
+	virtual FText GetDisplayName() const override;
+	virtual FText GetDescription() const override;
+	virtual FText GetHintText() const override;
+	virtual void GetAutoCompleteSuggestions(const TCHAR* Input, TArray<FString>& Out) override;
+	virtual void GetExecHistory(TArray<FString>& Out) override;
+	virtual bool Exec(const TCHAR* Input) override;
+	virtual bool AllowHotKeyClose() const override;
+	virtual bool AllowMultiLine() const override;
+	virtual FInputChord GetHotKey() const override;
+	//~ End IConsoleCommandExecutor Interface
+};
\ No newline at end of file
diff --git a/Engine/Plugins/Developer/Concert/ConcertApp/MultiUserServer/Source/MultiUserServer/Private/ConcertServerStyle.cpp b/Engine/Plugins/Developer/Concert/ConcertApp/MultiUserServer/Source/MultiUserServer/Private/ConcertServerStyle.cpp
new file mode 100644
--- /dev/null
+++ b/Engine/Plugins/Developer/Concert/ConcertApp/MultiUserServer/Source/MultiUserServer/Private/ConcertServerStyle.cpp
@@ -0,0 +1,73 @@
+// Copyright Epic Games, Inc. All Rights Reserved.
+
+#include "ConcertServerStyle.h"
+
+#include "EditorStyleSet.h"
+#include "Framework/Application/SlateApplication.h"
+#include "Interfaces/IPluginManager.h"
+#include "Styling/SlateStyleRegistry.h"
+
+TSharedPtr<FSlateStyleSet> FConcertServerStyle::StyleInstance = nullptr;
+
+void FConcertServerStyle::Initialize()
+{
+	if (!StyleInstance.IsValid())
+	{
+		StyleInstance = Create();
+		FSlateStyleRegistry::RegisterSlateStyle(*StyleInstance);
+	}
+}
+
+void FConcertServerStyle::Shutdown()
+{
+	if (StyleInstance)
+	{
+		
+		FSlateStyleRegistry::UnRegisterSlateStyle(*StyleInstance);
+		ensure(StyleInstance.IsUnique());
+		StyleInstance.Reset();
+	}
+}
+
+const ISlateStyle& FConcertServerStyle::Get()
+{
+	return *StyleInstance;
+}
+
+FName FConcertServerStyle::GetStyleSetName()
+{
+	static FName StyleSetName(TEXT("ConcertServerStyle"));
+	return StyleSetName;
+}
+
+#define IMAGE_PLUGIN_BRUSH( RelativePath, ... ) FSlateImageBrush( FConcertServerStyle::InContent( RelativePath, ".png" ), __VA_ARGS__ )
+#define IMAGE_BRUSH(RelativePath, ...) FSlateImageBrush(StyleSet->RootToContentDir(RelativePath, TEXT(".png")), __VA_ARGS__)
+#define IMAGE_PLUGIN_BRUSH_SVG( RelativePath, ... ) FSlateVectorImageBrush( FConcertServerStyle::InContent(RelativePath, ".svg"), __VA_ARGS__)
+
+FString FConcertServerStyle::InContent(const FString& RelativePath, const ANSICHAR* Extension)
+{
+	static FString ContentDir = IPluginManager::Get().FindPlugin(TEXT("MultiUserServer"))->GetContentDir();
+	return (ContentDir / RelativePath) + Extension;
+}
+
+TSharedRef<FSlateStyleSet> FConcertServerStyle::Create()
+{
+	TSharedRef<FSlateStyleSet> StyleSet = MakeShared<FSlateStyleSet>(GetStyleSetName());
+	
+	StyleSet->SetContentRoot(FPaths::EngineContentDir() / TEXT("Slate/Starship/Insights"));
+	StyleSet->SetCoreContentRoot(FPaths::EngineContentDir() / TEXT("Slate"));
+	
+	const FVector2D Icon12x12(12.0f, 12.0f); 
+	const FVector2D Icon16x16(16.0f, 16.0f); 
+	const FVector2D Icon20x20(20.0f, 20.0f); 
+	const FVector2D Icon32x32(20.0f, 20.0f); 
+	
+	StyleSet->Set("Concert.MultiUser", new IMAGE_PLUGIN_BRUSH("Icons/icon_MultiUser_32x", Icon32x32));
+	return StyleSet;
+}
+
+#undef TODO_IMAGE_BRUSH
+#undef EDITOR_BOX_BRUSH
+#undef EDITOR_IMAGE_BRUSH_SVG
+#undef EDITOR_IMAGE_BRUSH
+
diff --git a/Engine/Plugins/Developer/Concert/ConcertApp/MultiUserServer/Source/MultiUserServer/Private/ConcertServerStyle.h b/Engine/Plugins/Developer/Concert/ConcertApp/MultiUserServer/Source/MultiUserServer/Private/ConcertServerStyle.h
new file mode 100644
--- /dev/null
+++ b/Engine/Plugins/Developer/Concert/ConcertApp/MultiUserServer/Source/MultiUserServer/Private/ConcertServerStyle.h
@@ -0,0 +1,25 @@
+// Copyright Epic Games, Inc. All Rights Reserved.
+
+#pragma once
+
+#include "CoreMinimal.h"
+#include "Styling/SlateStyle.h"
+
+/** Style data for Insights tools */
+class FConcertServerStyle
+{
+public:
+	static void Initialize();
+	static void Shutdown();
+
+	static const ISlateStyle& Get();
+	static FName GetStyleSetName();
+
+private:
+	
+	static TSharedPtr<FSlateStyleSet> StyleInstance;
+	
+	static FString InContent(const FString& RelativePath, const ANSICHAR* Extension);
+	static TSharedRef<FSlateStyleSet> Create();
+
+};
diff --git a/Engine/Plugins/Developer/Concert/ConcertApp/MultiUserServer/Source/MultiUserServer/Private/MultiUserServerModule.cpp b/Engine/Plugins/Developer/Concert/ConcertApp/MultiUserServer/Source/MultiUserServer/Private/MultiUserServerModule.cpp
new file mode 100644
--- /dev/null
+++ b/Engine/Plugins/Developer/Concert/ConcertApp/MultiUserServer/Source/MultiUserServer/Private/MultiUserServerModule.cpp
@@ -0,0 +1,90 @@
+// Copyright Epic Games, Inc. All Rights Reserved.
+
+#include "IMultiUserServerModule.h"
+
+#include "ConcertConsoleCommandExecutor.h"
+#include "ConcertFrontendStyle.h"
+#include "ConcertServerStyle.h"
+#include "ConcertSyncServerLoopInitArgs.h"
+#include "Widgets/ConcertServerWindowController.h"
+
+#include "Features/IModularFeatures.h"
+#include "Framework/Application/SlateApplication.h"
+#include "StandaloneRenderer.h"
+#include "Misc/ConfigCacheIni.h"
+
+#define LOCTEXT_NAMESPACE "UnrealMultiUserUI"
+
+/**
+ * 
+ */
+class FConcertServerUIModule : public IMultiUserServerModule
+{
+public:
+	virtual void StartupModule() override 
+	{
+		MultiUserServerLayoutIni = GConfig->GetConfigFilename(TEXT("MultiUserServerLayout"));
+	}
+	
+	virtual void ShutdownModule() override
+	{
+		FConcertServerStyle::Shutdown();
+		
+		WindowController.Reset();
+		FSlateApplication::Shutdown();
+	}
+	
+	virtual void InitSlateForServer(FConcertSyncServerLoopInitArgs& InitArgs) override
+	{
+		if (ensureMsgf(!WindowController.IsValid(), TEXT("InitSlateForServer is designed to be called at most once to create UI to run alongside the Multi User server.")))
+		{
+			PreInitializeMultiUser();
+			InitArgs.PostInitServerLoop.AddRaw(this, &FConcertServerUIModule::InitializeSlateApplication);
+			InitArgs.TickPostGameThread.AddRaw(this, &FConcertServerUIModule::TickSlate);
+		}
+	}
+	
+private:
+
+	/** Config path storing layout config. */
+	FString MultiUserServerLayoutIni;
+
+	/** Handles execution of commands */
+	TUniquePtr<FConcertConsoleCommandExecutor> CommandExecutor;
+	
+	/** Creates and manages window. Only one instance per application. */
+	TSharedPtr<FConcertServerWindowController> WindowController;
+
+	void PreInitializeMultiUser()
+	{
+		FModuleManager::Get().LoadModuleChecked("EditorStyle");
+		FConcertServerStyle::Initialize();
+		FConcertFrontendStyle::Initialize();
+
+		// Log history must be initialized before MU server loop init prints any messages	
+		FModuleManager::Get().LoadModuleChecked("OutputLog");
+	}
+	
+	void InitializeSlateApplication(TSharedRef<IConcertSyncServer> SyncServer)
+	{
+		FSlateApplication::InitializeAsStandaloneApplication(GetStandardStandaloneRenderer());
+		const FText ApplicationTitle = LOCTEXT("AppTitle", "Unreal Multi User Server");
+		FGlobalTabmanager::Get()->SetApplicationTitle(ApplicationTitle);
+
+		CommandExecutor = MakeUnique<FConcertConsoleCommandExecutor>();
+		IModularFeatures::Get().RegisterModularFeature(IConsoleCommandExecutor::ModularFeatureName(), CommandExecutor.Get());
+        
+        WindowController = MakeShared<FConcertServerWindowController>(FConcertServerWindowInitParams{ SyncServer, MultiUserServerLayoutIni });
+		WindowController->CreateWindow();
+	}
+	
+	void TickSlate(double Tick)
+	{
+		FSlateApplication::Get().PumpMessages();
+		FSlateApplication::Get().Tick();
+	}
+};
+
+IMPLEMENT_MODULE(FConcertServerUIModule, ConcertSyncServerUI);
+
+#undef LOCTEXT_NAMESPACE
\ No newline at end of file
diff --git a/Engine/Plugins/Developer/Concert/ConcertApp/MultiUserServer/Source/MultiUserServer/Private/Widgets/Browser/ConcertServerSessionBrowserController.cpp b/Engine/Plugins/Developer/Concert/ConcertApp/MultiUserServer/Source/MultiUserServer/Private/Widgets/Browser/ConcertServerSessionBrowserController.cpp
new file mode 100644
--- /dev/null
+++ b/Engine/Plugins/Developer/Concert/ConcertApp/MultiUserServer/Source/MultiUserServer/Private/Widgets/Browser/ConcertServerSessionBrowserController.cpp
@@ -0,0 +1,215 @@
+﻿// Copyright Epic Games, Inc. All Rights Reserved.
+
+#include "ConcertServerSessionBrowserController.h"
+
+#include "ConcertServerEvents.h"
+#include "ConcertServerStyle.h"
+#include "IConcertServer.h"
+#include "IConcertSyncServer.h"
+#include "Framework/Docking/TabManager.h"
+#include "SConcertServerSessionBrowser.h"
+#include "Framework/Notifications/NotificationManager.h"
+#include "Misc/AsyncTaskNotification.h"
+#include "Session/Browser/ConcertSessionItem.h"
+#include "Textures/SlateIcon.h"
+#include "Widgets/ConcertServerTabs.h"
+#include "Widgets/ConcertServerWindowController.h"
+#include "Widgets/Docking/SDockTab.h"
+#include "Widgets/Notifications/SNotificationList.h"
+
+#define LOCTEXT_NAMESPACE "UnrealMultiUserUI"
+
+int32 FConcertServerSessionBrowserController::GetNumConnectedClients(const FGuid& SessionId) const
+{
+	TSharedPtr<IConcertServerSession> Session = ServerInstance->GetConcertServer()->GetLiveSession(SessionId);
+	return ensure(Session) ? Session->GetSessionClients().Num() : 0;
+}
+
+FConcertServerSessionBrowserController::~FConcertServerSessionBrowserController()
+{
+	ConcertServerEvents::OnLiveSessionCreated().RemoveAll(this);
+	ConcertServerEvents::OnLiveSessionDestroyed().RemoveAll(this);
+	ConcertServerEvents::OnArchivedSessionCreated().RemoveAll(this);
+	ConcertServerEvents::OnArchivedSessionDestroyed().RemoveAll(this);
+}
+
+void FConcertServerSessionBrowserController::Init(const FConcertComponentInitParams& Params)
+{
+	ServerInstance = Params.Server;
+	Owner = Params.WindowController;
+	FGlobalTabmanager::Get()->RegisterTabSpawner(
+			ConcertServerTabs::GetSessionBrowserTabId(),
+			FOnSpawnTab::CreateRaw(this, &FConcertServerSessionBrowserController::SpawnSessionBrowserTab)
+		)
+		.SetDisplayName(LOCTEXT("SessionBrowserTabTitle", "Session Browser"))
+		.SetTooltipText(LOCTEXT("SessionBrowserTooltipText", "A section to browse, start, archive, and restore server sessions."))
+		.SetIcon(FSlateIcon(FConcertServerStyle::GetStyleSetName(), TEXT("Concert.MultiUser")));
+
+	ConcertServerEvents::OnLiveSessionCreated().AddSP(this, &FConcertServerSessionBrowserController::OnLiveSessionCreated);
+	ConcertServerEvents::OnLiveSessionDestroyed().AddSP(this, &FConcertServerSessionBrowserController::OnLiveSessionDestroyed);
+	ConcertServerEvents::OnArchivedSessionCreated().AddSP(this, &FConcertServerSessionBrowserController::OnArchivedSessionCreated);
+	ConcertServerEvents::OnArchivedSessionDestroyed().AddSP(this, &FConcertServerSessionBrowserController::OnArchivedSessionDestroyed);
+}
+
+TArray<FConcertServerInfo> FConcertServerSessionBrowserController::GetServers() const
+{
+	return { ServerInstance->GetConcertServer()->GetServerInfo() };
+}
+
+TArray<IConcertSessionBrowserController::FActiveSessionInfo> FConcertServerSessionBrowserController::GetActiveSessions() const
+{
+	const FConcertServerInfo& ServerInfo = ServerInstance->GetConcertServer()->GetServerInfo();
+	const TArray<TSharedPtr<IConcertServerSession>> ServerSessions = ServerInstance->GetConcertServer()->GetLiveSessions();
+	
+	TArray<FActiveSessionInfo> Result;
+	Result.Reserve(ServerSessions.Num());
+	for (const TSharedPtr<IConcertServerSession>& LiveSession : ServerSessions)
+	{
+		FActiveSessionInfo Info{ ServerInfo, LiveSession->GetSessionInfo(), LiveSession->GetSessionClients() };
+		Result.Add(Info);
+	}
+
+	return Result;
+}
+
+TArray<IConcertSessionBrowserController::FArchivedSessionInfo> FConcertServerSessionBrowserController::GetArchivedSessions() const
+{
+	const FConcertServerInfo& ServerInfo = ServerInstance->GetConcertServer()->GetServerInfo();
+	const TArray<FConcertSessionInfo> ConcertSessionInfos = ServerInstance->GetConcertServer()->GetArchivedSessionInfos();
+	
+	TArray<FArchivedSessionInfo> Result;
+	Result.Reserve(ConcertSessionInfos.Num());
+	for (const FConcertSessionInfo& SessionInfo : ConcertSessionInfos)
+	{
+		Result.Emplace(ServerInfo, SessionInfo);
+	}
+
+	return Result;
+}
+
+TOptional<FConcertSessionInfo> FConcertServerSessionBrowserController::GetActiveSessionInfo(const FGuid& AdminEndpoint, const FGuid& SessionId) const
+{
+	const TSharedPtr<IConcertServerSession> LiveSession = ServerInstance->GetConcertServer()->GetLiveSession(SessionId);
+	return LiveSession ? LiveSession->GetSessionInfo() : TOptional<FConcertSessionInfo>{};
+}
+
+TOptional<FConcertSessionInfo> FConcertServerSessionBrowserController::GetArchivedSessionInfo(const FGuid& AdminEndpoint, const FGuid& SessionId) const
+{
+	return ServerInstance->GetConcertServer()->GetArchivedSessionInfo(SessionId);
+}
+
+void FConcertServerSessionBrowserController::CreateSession(const FGuid& ServerAdminEndpointId, const FString& SessionName)
+{
+	FConcertSessionInfo SessionInfo = ServerInstance->GetConcertServer()->CreateSessionInfo();
+	SessionInfo.SessionName = SessionName;
+	SessionInfo.Settings.Initialize();
+	FConcertSessionVersionInfo VersionInfo;
+	VersionInfo.Initialize();
+	SessionInfo.VersionInfos.Emplace(VersionInfo);
+	
+	FText FailureReason = FText::GetEmpty();
+	const bool bSuccess = ServerInstance->GetConcertServer()->CreateSession(SessionInfo, FailureReason) != nullptr;
+	
+	NotifyUserOfFinishedSessionAction(bSuccess,
+		bSuccess ? FText::Format(LOCTEXT("CreatedSessionFmt", "Created Session '{0}'"), FText::FromString(SessionName)) : FText::Format(LOCTEXT("FailedToCreateSessionFmt", "Failed to create Session '{0}'"), FText::FromString(SessionName)),
+		FailureReason
+		);
+}
+
+void FConcertServerSessionBrowserController::ArchiveSession(const FGuid& ServerAdminEndpointId, const FGuid& SessionId, const FString& ArchiveName, const FConcertSessionFilter& SessionFilter)
+{
+	FText FailureReason = FText::GetEmpty();
+	const bool bSuccess = ServerInstance->GetConcertServer()->ArchiveSession(SessionId, ArchiveName, SessionFilter, FailureReason).IsValid();
+	
+	NotifyUserOfFinishedSessionAction(bSuccess,
+		bSuccess ? FText::Format(LOCTEXT("ArchivedSessionFmt", "Archived Session '{0}'"),  FText::FromString(ArchiveName)) : FText::Format(LOCTEXT("FailedToArchivedSessionFmt", "Failed to archive Session '{0}'"), FText::FromString(ArchiveName)),
+		FailureReason
+		);
+}
+
+void FConcertServerSessionBrowserController::RestoreSession(const FGuid& ServerAdminEndpointId, const FGuid& SessionId, const FString& RestoredName, const FConcertSessionFilter& SessionFilter)
+{
+	if (const TOptional<FConcertSessionInfo> SessionInfo = GetArchivedSessionInfo(ServerAdminEndpointId, SessionId))
+	{
+		FText FailureReason = FText::GetEmpty();
+		const bool bSuccess = ServerInstance->GetConcertServer()->RestoreSession(SessionId, *SessionInfo, SessionFilter, FailureReason).IsValid();
+		NotifyUserOfFinishedSessionAction(bSuccess,
+			bSuccess ? FText::Format(LOCTEXT("RestoreSessionFmt", "Restored Session '{0}'"),  FText::FromString(RestoredName)) : FText::Format(LOCTEXT("FailedToRestoreSessionFmt", "Failed to restore Session '{0}'"), FText::FromString(RestoredName)),
+			FailureReason
+			);
+	}
+}
+
+TSharedRef<SDockTab> FConcertServerSessionBrowserController::SpawnSessionBrowserTab(const FSpawnTabArgs& Args)
+{
+	const TSharedRef<SDockTab> DockTab = SNew(SDockTab)
+		.Label(LOCTEXT("SessionBrowserTabTitle", "Sessions"))
+		.TabRole(MajorTab)
+		[
+			SAssignNew(ConcertBrowser, SConcertServerSessionBrowser, SharedThis(this))
+				.DoubleClickLiveSession(this, &FConcertServerSessionBrowserController::OpenSession)
+				.DoubleClickArchivedSession(this, &FConcertServerSessionBrowserController::OpenSession)
+		];
+
+	FGlobalTabmanager::Get()->SetMainTab(DockTab);
+	return DockTab;
+}
+
+void FConcertServerSessionBrowserController::RefreshSessionList()
+{
+	if (ConcertBrowser.IsValid())
+	{
+		// 1. Accumulate several events in same tick so list is only updated once
+		// 2. Several events, e.g. OnLiveSessionCreated, are called before the session creation is visible to getter functions
+		ConcertBrowser->RequestRefreshListNextTick();
+	}
+}
+
+void FConcertServerSessionBrowserController::OpenSession(const TSharedPtr<FConcertSessionItem>& SessionItem)
+{
+	Owner.Pin()->OpenSessionTab(SessionItem->SessionId);
+}
+
+void FConcertServerSessionBrowserController::RenameSession(const FGuid& ServerAdminEndpointId, const FGuid& SessionId, const FString& NewName)
+{
+	if (const TOptional<FConcertSessionInfo> SessionInfo = GetArchivedSessionInfo(ServerAdminEndpointId, SessionId))
+	{
+		FText FailureReason = FText::GetEmpty();
+		const bool bSuccess = ServerInstance->GetConcertServer()->RenameSession(SessionId, NewName, FailureReason);
+		NotifyUserOfFinishedSessionAction(bSuccess,
+				bSuccess ? FText::Format(LOCTEXT("RenameSessionFmt", "Rename Session '{0}' as '{1}'"), FText::FromString(SessionInfo->SessionName), FText::FromString(NewName))
+					: FText::Format(LOCTEXT("FailedToArchivedSessionFmt", "Failed to rename Session '{0}' as '{1}'"), FText::FromString(SessionInfo->SessionName), FText::FromString(NewName)),
+				FailureReason
+				);
+	}
+}
+
+void FConcertServerSessionBrowserController::DeleteSession(const FGuid& ServerAdminEndpointId, const FGuid& SessionId)
+{
+	const TOptional<FConcertSessionInfo> ArchivedInfo = GetArchivedSessionInfo(ServerAdminEndpointId, SessionId);
+	const TOptional<FConcertSessionInfo> LiveInfo = GetActiveSessionInfo(ServerAdminEndpointId, SessionId);
+	const TOptional<FString> SessionName = ArchivedInfo ? ArchivedInfo->SessionName : LiveInfo ? LiveInfo->SessionName : TOptional<FString>{};
+	if (SessionName)
+	{
+		FText FailureReason;
+		const bool bSuccess = ServerInstance->GetConcertServer()->DestroySession(SessionId, FailureReason);
+		NotifyUserOfFinishedSessionAction(bSuccess,
+			bSuccess ? FText::Format(LOCTEXT("ArchivedSessionFmt", "Deleted Session '{0}'"), FText::FromString(*SessionName)) : FText::Format(LOCTEXT("FailedToArchivedSessionFmt", "Failed to archive Session '{0}'"), FText::FromString(LiveInfo->SessionName)),
+			FailureReason
+			);
+	}
+}
+
+void FConcertServerSessionBrowserController::NotifyUserOfFinishedSessionAction(const bool bSuccess, const FText& Title, const FText& Details)
+{
+	FNotificationInfo NotificationInfo(Title);
+	NotificationInfo.SubText = Details;
+	FSlateNotificationManager::Get().AddNotification(NotificationInfo);
+
+	if (bSuccess)
+	{
+		ConcertBrowser->RequestRefreshListNextTick();
+	}
+}
+
+#undef LOCTEXT_NAMESPACE 
diff --git a/Engine/Plugins/Developer/Concert/ConcertApp/MultiUserServer/Source/MultiUserServer/Private/Widgets/Browser/ConcertServerSessionBrowserController.h b/Engine/Plugins/Developer/Concert/ConcertApp/MultiUserServer/Source/MultiUserServer/Private/Widgets/Browser/ConcertServerSessionBrowserController.h
new file mode 100644
--- /dev/null
+++ b/Engine/Plugins/Developer/Concert/ConcertApp/MultiUserServer/Source/MultiUserServer/Private/Widgets/Browser/ConcertServerSessionBrowserController.h
@@ -0,0 +1,82 @@
+﻿// Copyright Epic Games, Inc. All Rights Reserved.
+
+#pragma once
+
+#include "CoreMinimal.h"
+#include "Session/Browser/IConcertSessionBrowserController.h"
+#include "Widgets/IConcertComponent.h"
+
+class IConcertServerSession;
+class IConcertServer;
+class FConcertSessionItem;
+class FSpawnTabArgs;
+class SDockTab;
+class SConcertServerSessionBrowser;
+
+/** Implements controller in model-view-controller pattern. */
+class FConcertServerSessionBrowserController
+	:
+	public TSharedFromThis<FConcertServerSessionBrowserController>,
+	public IConcertComponent,
+	public IConcertSessionBrowserController
+{
+public:
+
+	int32 GetNumConnectedClients(const FGuid& SessionId) const;
+
+	~FConcertServerSessionBrowserController();
+	
+	//~ Begin IConcertComponent Interface
+	virtual void Init(const FConcertComponentInitParams& Params) override;
+	//~ End IConcertComponent Interface
+	
+	//~ Begin IConcertSessionBrowserController Interface
+	virtual TArray<FConcertServerInfo> GetServers() const override;
+	virtual TArray<FActiveSessionInfo> GetActiveSessions() const override;
+	virtual TArray<FArchivedSessionInfo> GetArchivedSessions() const override;
+	virtual TOptional<FConcertSessionInfo> GetActiveSessionInfo(const FGuid& AdminEndpoint, const FGuid& SessionId) const override;
+	virtual TOptional<FConcertSessionInfo> GetArchivedSessionInfo(const FGuid& AdminEndpoint, const FGuid& SessionId) const override;
+	
+	virtual void CreateSession(const FGuid& ServerAdminEndpointId, const FString& SessionName) override;
+	virtual void ArchiveSession(const FGuid& ServerAdminEndpointId, const FGuid& SessionId, const FString& ArchiveName, const FConcertSessionFilter& SessionFilter) override;
+	virtual void RestoreSession(const FGuid& ServerAdminEndpointId, const FGuid& SessionId, const FString& RestoredName, const FConcertSessionFilter& SessionFilter) override;
+	virtual void RenameActiveSession(const FGuid& ServerAdminEndpointId, const FGuid& SessionId, const FString& NewName) override { RenameSession(ServerAdminEndpointId, SessionId, NewName); }
+	virtual void RenameArchivedSession(const FGuid& ServerAdminEndpointId, const FGuid& SessionId, const FString& NewName) override { RenameSession(ServerAdminEndpointId, SessionId, NewName); }
+	virtual void DeleteActiveSession(const FGuid& ServerAdminEndpointId, const FGuid& SessionId) override { DeleteSession(ServerAdminEndpointId, SessionId); }
+	virtual void DeleteArchivedSession(const FGuid& ServerAdminEndpointId, const FGuid& SessionId) override { DeleteSession(ServerAdminEndpointId, SessionId); }
+	
+	// The server operator always has permission for these actions:
+	virtual bool CanRenameActiveSession(const FGuid& ServerAdminEndpointId, const FGuid& SessionId) const override { return true; }
+	virtual bool CanRenameArchivedSession(const FGuid& ServerAdminEndpointId, const FGuid& SessionId) const override { return true; }
+	virtual bool CanDeleteActiveSession(const FGuid& ServerAdminEndpointId, const FGuid& SessionId) const override { return true; }
+	virtual bool CanDeleteArchivedSession(const FGuid& ServerAdminEndpointId, const FGuid& SessionId) const override { return true; }
+	//~ Begin IConcertSessionBrowserController Interface
+
+private:
+
+	/** Server instance we're representing */
+	TSharedPtr<IConcertSyncServer> ServerInstance;
+	/** Used to open selected sessions */
+	TWeakPtr<FConcertServerWindowController> Owner;
+	
+	TSharedPtr<SConcertServerSessionBrowser> ConcertBrowser;
+	
+	TSharedRef<SDockTab> SpawnSessionBrowserTab(const FSpawnTabArgs& Args);
+
+	// Update view when session list changes
+	void OnLiveSessionCreated(bool, const IConcertServer&, TSharedRef<IConcertServerSession>) { RefreshSessionList(); }
+	void OnLiveSessionDestroyed(const IConcertServer&, TSharedRef<IConcertServerSession>) { RefreshSessionList(); }
+	void OnArchivedSessionCreated(bool, const IConcertServer&, const FString&, const FConcertSessionInfo&) { RefreshSessionList(); }
+	void OnArchivedSessionDestroyed(const IConcertServer&, const FGuid&) { RefreshSessionList(); }
+	
+	void RefreshSessionList();
+
+	// Session actions
+	void OpenSession(const TSharedPtr<FConcertSessionItem>& SessionItem);
+	void RenameSession(const FGuid& ServerAdminEndpointId, const FGuid& SessionId, const FString& NewName);
+	void DeleteSession(const FGuid& ServerAdminEndpointId, const FGuid& SessionId);
+
+	// Notifications
+	void NotifyUserOfFinishedSessionAction(const bool bSuccess, const FText& Title) { NotifyUserOfFinishedSessionAction(bSuccess, Title, FText::GetEmpty()); }
+	void NotifyUserOfFinishedSessionAction(const bool bSuccess, const FText& Title, const FText& Details);
+};
diff --git a/Engine/Plugins/Developer/Concert/ConcertApp/MultiUserServer/Source/MultiUserServer/Private/Widgets/Browser/SConcertServerSessionBrowser.cpp b/Engine/Plugins/Developer/Concert/ConcertApp/MultiUserServer/Source/MultiUserServer/Private/Widgets/Browser/SConcertServerSessionBrowser.cpp
new file mode 100644
--- /dev/null
+++ b/Engine/Plugins/Developer/Concert/ConcertApp/MultiUserServer/Source/MultiUserServer/Private/Widgets/Browser/SConcertServerSessionBrowser.cpp
@@ -0,0 +1,176 @@
+﻿// Copyright Epic Games, Inc. All Rights Reserved.
+
+#include "SConcertServerSessionBrowser.h"
+
+#include "Session/Browser/SConcertSessionBrowser.h"
+#include "Widgets/Browser/ConcertServerSessionBrowserController.h"
+#include "Widgets/StatusBar/SConcertStatusBar.h"
+
+#include "Dialog/SMessageDialog.h"
+#include "Framework/Docking/TabManager.h"
+#include "Session/Browser/ConcertSessionItem.h"
+#include "Widgets/ConcertServerTabs.h"
+#include "Widgets/Input/SSearchBox.h"
+#include "Widgets/Layout/SBorder.h"
+#include "Widgets/SBoxPanel.h"
+
+#define LOCTEXT_NAMESPACE "UnrealMultiUserUI"
+
+SConcertServerSessionBrowser::~SConcertServerSessionBrowser()
+{
+	UnregisterFromOnRootWindowClosed();
+}
+
+void SConcertServerSessionBrowser::Construct(const FArguments& InArgs, TSharedRef<FConcertServerSessionBrowserController> InController)
+{
+	Controller = InController;
+	ChildSlot
+	[
+		SNew(SBorder)
+		.BorderImage(FAppStyle::Get().GetBrush("ToolPanel.GroupBorder"))
+		.Padding(FMargin(1.0f, 2.0f))
+		[
+			SNew(SVerticalBox)
+
+			// Session list.
+			+SVerticalBox::Slot()
+			.FillHeight(1.0f)
+			.Padding(1.0f, 2.0f)
+			[
+				MakeSessionTableView(InArgs)
+			]
+
+			+SVerticalBox::Slot()
+			.AutoHeight()
+			.HAlign(HAlign_Fill)
+			.VAlign(VAlign_Bottom)
+			[
+				SNew(SConcertStatusBar, ConcertServerTabs::GetSessionBrowserTabId())
+			]
+		]
+	];
+}
+
+TSharedRef<SWidget> SConcertServerSessionBrowser::MakeSessionTableView(const FArguments& InArgs)
+{
+	SearchText = MakeShared<FText>();
+	return SAssignNew(SessionBrowser, SConcertSessionBrowser, Controller.Pin().ToSharedRef(), SearchText)
+		.OnLiveSessionDoubleClicked(InArgs._DoubleClickLiveSession)
+		.OnArchivedSessionDoubleClicked(InArgs._DoubleClickArchivedSession)
+		.OnRequestedDeleteSession(this, &SConcertServerSessionBrowser::RequestDeleteSession)
+		// Pretend a modal dialog said no - RequestDeleteSession will show non-modal dialog
+		.CanDeleteArchivedSession_Lambda([](TSharedPtr<FConcertSessionItem>) { return false; })
+		.CanDeleteActiveSession_Lambda([](TSharedPtr<FConcertSessionItem>) { return false; });
+}
+
+void SConcertServerSessionBrowser::RequestDeleteSession(const TSharedPtr<FConcertSessionItem>& SessionItem) 
+{
+	if (const TSharedPtr<SMessageDialog> PinnedDeleteSessionDialog = DeleteSessionDialog.Pin())
+	{
+		// Force user to answer previous question first
+		PinnedDeleteSessionDialog->BringToFront();
+		PinnedDeleteSessionDialog->FlashWindow();
+	}
+	else
+	{
+		// Modal dialog would cause server to halt thread and stop communicating with connected clients
+		switch (SessionItem->Type)
+		{
+		case FConcertSessionItem::EType::ActiveSession:
+			DeleteActiveSessionWithNonModalQuestion(SessionItem); break;
+		case FConcertSessionItem::EType::ArchivedSession:
+			DeleteArchivedSessionWithNonModalQuestion(SessionItem); break;
+		default:
+			break;
+		}
+
+		// The dialog should be closed when the root window is closed
+		if (const TSharedPtr<SWindow> RootWindow = FGlobalTabmanager::Get()->GetRootWindow();
+			RootWindow && !RootWindow->GetOnWindowClosedEvent().IsBoundToObject(this))
+		{
+			RootWindow->GetOnWindowClosedEvent().AddSP(this, &SConcertServerSessionBrowser::OnRootWindowClosed);
+		}
+	}
+}
+
+void SConcertServerSessionBrowser::OnRootWindowClosed(const TSharedRef<SWindow>&) const
+{
+	if (const TSharedPtr<SMessageDialog> PinnedDeleteSessionDialog = DeleteSessionDialog.Pin())
+	{
+		PinnedDeleteSessionDialog->RequestDestroyWindow();
+	}
+	
+	UnregisterFromOnRootWindowClosed();
+}
+
+void SConcertServerSessionBrowser::UnregisterFromOnRootWindowClosed() const
+{
+	if (const TSharedPtr<SWindow> RootWindow = FGlobalTabmanager::Get()->GetRootWindow())
+	{
+		RootWindow->GetOnWindowClosedEvent().RemoveAll(this);
+	}
+}
+
+void SConcertServerSessionBrowser::DeleteArchivedSessionWithNonModalQuestion(const TSharedPtr<FConcertSessionItem>& SessionItem)
+{
+	const FText Message = FText::Format(
+	LOCTEXT("DeleteArchivedDescription", "Deleting a session will cause all associated data to be removed.\n\nDelete {0}?"),
+		FText::FromString(SessionItem->SessionName)
+		);
+
+	auto DeleteArchived = [WeakController = TWeakPtr<FConcertServerSessionBrowserController>(Controller), WeakSessionItem = TWeakPtr<FConcertSessionItem>(SessionItem)]()
+	{
+		const TSharedPtr<FConcertServerSessionBrowserController> PinnedController = WeakController.Pin();
+		const TSharedPtr<FConcertSessionItem> PinnedSessionItem = WeakSessionItem.Pin();
+		if (PinnedController && PinnedSessionItem)
+		{
+			PinnedController->DeleteArchivedSession(PinnedSessionItem->ServerAdminEndpointId, PinnedSessionItem->SessionId);
+		}
+	};
+	const TSharedRef<SMessageDialog> Dialog = SNew(SMessageDialog)
+		.Title(LOCTEXT("DisconnectUsersTitle", "Delete session?"))
+		.IconBrush("Icons.WarningWithColor.Large")
+		.Message(Message)
+		.Buttons({
+			SMessageDialog::FButton(LOCTEXT("DeleteArchivedButton", "Delete"))
+				.SetOnClicked(FSimpleDelegate::CreateLambda(DeleteArchived)),
+			SMessageDialog::FButton(LOCTEXT("CancelButton", "Cancel")).SetPrimary(true)
+		});
+	
+	DeleteSessionDialog = Dialog;
+	Dialog->Show();
+}
+
+void SConcertServerSessionBrowser::DeleteActiveSessionWithNonModalQuestion(const TSharedPtr<FConcertSessionItem>& SessionItem)
+{
+	const int32 NumUsers = Controller.Pin()->GetNumConnectedClients(SessionItem->SessionId);
+	const FText Message = FText::Format(
+		LOCTEXT("DeletedActiveDescription", "There {0}|plural(one=is,other=are) {0} connected {0}|plural(one=client,other=clients) in the current session.\nDeleting a session will force all connected clients to disconnect.\n\nDelete {1}?"),
+		NumUsers,
+		FText::FromString(SessionItem->SessionName)
+		);
+
+	auto DeleteActive = [WeakController = TWeakPtr<FConcertServerSessionBrowserController>(Controller), WeakSessionItem = TWeakPtr<FConcertSessionItem>(SessionItem)]()
+	{
+		const TSharedPtr<FConcertServerSessionBrowserController> PinnedController = WeakController.Pin();
+		const TSharedPtr<FConcertSessionItem> PinnedSessionItem = WeakSessionItem.Pin();
+		if (PinnedController && PinnedSessionItem)
+		{
+			PinnedController->DeleteActiveSession(PinnedSessionItem->ServerAdminEndpointId, PinnedSessionItem->SessionId);
+		}
+	};
+	const TSharedRef<SMessageDialog> Dialog = SNew(SMessageDialog)
+		.Title(LOCTEXT("DisconnectUsersTitle", "Force Users to Disconnect?"))
+		.IconBrush("Icons.WarningWithColor.Large")
+		.Message(Message)
+		.Buttons({
+			SMessageDialog::FButton(LOCTEXT("DeleteActiveButton", "Delete")).SetPrimary(true)
+				.SetOnClicked(FSimpleDelegate::CreateLambda(DeleteActive)),
+			SMessageDialog::FButton(LOCTEXT("CancelButton", "Cancel"))
+		});
+	
+	DeleteSessionDialog = Dialog;
+	Dialog->Show();
+}
+
+#undef LOCTEXT_NAMESPACE
diff --git a/Engine/Plugins/Developer/Concert/ConcertApp/MultiUserServer/Source/MultiUserServer/Private/Widgets/Browser/SConcertServerSessionBrowser.h b/Engine/Plugins/Developer/Concert/ConcertApp/MultiUserServer/Source/MultiUserServer/Private/Widgets/Browser/SConcertServerSessionBrowser.h
new file mode 100644
--- /dev/null
+++ b/Engine/Plugins/Developer/Concert/ConcertApp/MultiUserServer/Source/MultiUserServer/Private/Widgets/Browser/SConcertServerSessionBrowser.h
@@ -0,0 +1,58 @@
+﻿// Copyright Epic Games, Inc. All Rights Reserved.
+
+#pragma once
+
+#include "CoreMinimal.h"
+#include "Session/Browser/SConcertSessionBrowser.h"
+#include "Widgets/DeclarativeSyntaxSupport.h"
+#include "Widgets/SCompoundWidget.h"
+
+class SMessageDialog;
+class FConcertServerSessionBrowserController;
+class SSearchBox;
+
+/** Shows a list of server sessions */
+class SConcertServerSessionBrowser : public SCompoundWidget
+{
+public:
+	
+	SLATE_BEGIN_ARGS(SConcertServerSessionBrowser) { }
+		SLATE_EVENT(FSessionDelegate, DoubleClickLiveSession)
+		SLATE_EVENT(FSessionDelegate, DoubleClickArchivedSession)
+	SLATE_END_ARGS()
+	virtual ~SConcertServerSessionBrowser() override;
+
+	void Construct(const FArguments& InArgs, TSharedRef<FConcertServerSessionBrowserController> InController);
+
+	void RequestRefreshListNextTick() { bRequestedRefresh = true; }
+	virtual void Tick(const FGeometry& AllottedGeometry, const double InCurrentTime, const float InDeltaTime) override
+	{
+		if (bRequestedRefresh)
+		{
+			SessionBrowser->RefreshSessionList();
+			bRequestedRefresh = false;
+		}
+
+		SCompoundWidget::Tick(AllottedGeometry, InCurrentTime, InDeltaTime);
+	}
+	
+private:
+
+	/** We can ask the controller about information and notify it about UI events. */
+	TWeakPtr<FConcertServerSessionBrowserController> Controller;
+	/** Tracks whether there is a dialog asking the user to delete a session. Used to avoid opening multiple. */
+	TWeakPtr<SMessageDialog> DeleteSessionDialog;
+
+	bool bRequestedRefresh = false;
+
+	TSharedPtr<FText> SearchText;
+	TSharedPtr<SConcertSessionBrowser> SessionBrowser;
+	
+	TSharedRef<SWidget> MakeSessionTableView(const FArguments& InArgs);
+
+	void RequestDeleteSession(const TSharedPtr<FConcertSessionItem>& SessionItem);
+	void OnRootWindowClosed(const TSharedRef<SWindow>&) const;
+	void UnregisterFromOnRootWindowClosed() const;
+	void DeleteArchivedSessionWithNonModalQuestion(const TSharedPtr<FConcertSessionItem>& SessionItem);
+	void DeleteActiveSessionWithNonModalQuestion(const TSharedPtr<FConcertSessionItem>& SessionItem);
+};
diff --git a/Engine/Plugins/Developer/Concert/ConcertApp/MultiUserServer/Source/MultiUserServer/Private/Widgets/ConcertServerTabs.cpp b/Engine/Plugins/Developer/Concert/ConcertApp/MultiUserServer/Source/MultiUserServer/Private/Widgets/ConcertServerTabs.cpp
new file mode 100644
--- /dev/null
+++ b/Engine/Plugins/Developer/Concert/ConcertApp/MultiUserServer/Source/MultiUserServer/Private/Widgets/ConcertServerTabs.cpp
@@ -0,0 +1,9 @@
+﻿// Copyright Epic Games, Inc. All Rights Reserved.
+
+#include "ConcertServerTabs.h"
+
+FName ConcertServerTabs::GetSessionBrowserTabId()
+{
+	static FName SessionBrowserTabId("SessionBrowserTabId");
+	return SessionBrowserTabId;
+}
\ No newline at end of file
diff --git a/Engine/Plugins/Developer/Concert/ConcertApp/MultiUserServer/Source/MultiUserServer/Private/Widgets/ConcertServerTabs.h b/Engine/Plugins/Developer/Concert/ConcertApp/MultiUserServer/Source/MultiUserServer/Private/Widgets/ConcertServerTabs.h
new file mode 100644
--- /dev/null
+++ b/Engine/Plugins/Developer/Concert/ConcertApp/MultiUserServer/Source/MultiUserServer/Private/Widgets/ConcertServerTabs.h
@@ -0,0 +1,10 @@
+﻿// Copyright Epic Games, Inc. All Rights Reserved.
+
+#pragma once
+
+#include "CoreMinimal.h"
+
+namespace ConcertServerTabs
+{
+	FName GetSessionBrowserTabId();
+}
diff --git a/Engine/Plugins/Developer/Concert/ConcertApp/MultiUserServer/Source/MultiUserServer/Private/Widgets/ConcertServerWindowController.cpp b/Engine/Plugins/Developer/Concert/ConcertApp/MultiUserServer/Source/MultiUserServer/Private/Widgets/ConcertServerWindowController.cpp
new file mode 100644
--- /dev/null
+++ b/Engine/Plugins/Developer/Concert/ConcertApp/MultiUserServer/Source/MultiUserServer/Private/Widgets/ConcertServerWindowController.cpp
@@ -0,0 +1,134 @@
+// Copyright Epic Games, Inc. All Rights Reserved.
+
+#include "ConcertServerWindowController.h"
+
+#include "Browser/ConcertServerSessionBrowserController.h"
+#include "ConcertServerTabs.h"
+#include "IConcertServer.h"
+#include "IConcertSyncServer.h"
+#include "SessionTabs/Archived/ArchivedConcertSessionTab.h"
+#include "SessionTabs/Live/LiveConcertSessionTab.h"
+
+#include "Framework/Application/SlateApplication.h"
+#include "Framework/Docking/TabManager.h"
+#include "Framework/Docking/LayoutService.h"
+#include "Framework/Notifications/NotificationManager.h"
+#include "HAL/PlatformApplicationMisc.h"
+#include "Misc/ConfigCacheIni.h"
+#include "OutputLog/Public/OutputLogModule.h"
+
+
+#include "Widgets/SWindow.h"
+
+#define LOCTEXT_NAMESPACE "UnrealMultiUserUI"
+
+FConcertServerWindowController::FConcertServerWindowController(const FConcertServerWindowInitParams& Params)
+	: MultiUserServerLayoutIni(Params.MultiUserServerLayoutIni)
+{
+	ServerInstance = Params.Server;
+	SessionBrowserController = MakeShared<FConcertServerSessionBrowserController>();
+}
+
+void FConcertServerWindowController::CreateWindow()
+{
+	InitComponents();
+
+	FDisplayMetrics DisplayMetrics;
+	FSlateApplication::Get().GetDisplayMetrics(DisplayMetrics);
+	const float DPIScaleFactor = FPlatformApplicationMisc::GetDPIScaleFactorAtPoint(DisplayMetrics.PrimaryDisplayWorkAreaRect.Left, DisplayMetrics.PrimaryDisplayWorkAreaRect.Top);
+
+	const bool bEmbedTitleAreaContent = false;
+	const FVector2D ClientSize(960.0f * DPIScaleFactor, 640.0f * DPIScaleFactor);
+	TSharedRef<SWindow> RootWindowRef = SNew(SWindow)
+		.Title(LOCTEXT("WindowTitle", "Unreal Multi User Server"))
+		.CreateTitleBar(!bEmbedTitleAreaContent)
+		.SupportsMaximize(true)
+		.SupportsMinimize(true)
+		.IsInitiallyMaximized(false)
+		.IsInitiallyMinimized(false)
+		.SizingRule(ESizingRule::UserSized)
+		.AutoCenter(EAutoCenter::PreferredWorkArea)
+		.ClientSize(ClientSize)
+		.AdjustInitialSizeAndPositionForDPIScale(false);
+	RootWindow = RootWindowRef;
+		
+	const bool bShowRootWindowImmediately = false;
+	FSlateApplication::Get().AddWindow(RootWindowRef, bShowRootWindowImmediately);
+	FGlobalTabmanager::Get()->SetRootWindow(RootWindowRef);
+	FGlobalTabmanager::Get()->SetAllowWindowMenuBar(true);
+
+	FSlateNotificationManager::Get().SetRootWindow(RootWindowRef);
+	TSharedRef<FTabManager::FLayout> DefaultLayout = FTabManager::NewLayout("UnrealMultiUserServerLayout_v1.0");
+	DefaultLayout->AddArea
+	(
+		FTabManager::NewPrimaryArea()
+		->Split
+		(
+			FTabManager::NewStack()
+				->AddTab(ConcertServerTabs::GetSessionBrowserTabId(), ETabState::OpenedTab)
+				->SetForegroundTab(ConcertServerTabs::GetSessionBrowserTabId())
+		)
+	);
+	
+	PersistentLayout = FLayoutSaveRestore::LoadFromConfig(MultiUserServerLayoutIni, DefaultLayout);
+	TSharedPtr<SWidget> Content = FGlobalTabmanager::Get()->RestoreFrom(PersistentLayout.ToSharedRef(), RootWindow, bEmbedTitleAreaContent, EOutputCanBeNullptr::Never);
+	RootWindow->SetContent(Content.ToSharedRef());
+
+	RootWindow->SetOnWindowClosed(FOnWindowClosed::CreateRaw(this, &FConcertServerWindowController::OnWindowClosed));
+	RootWindow->ShowWindow();
+	constexpr bool bForceWindowToFront = true;
+	RootWindow->BringToFront(bForceWindowToFront);
+}
+
+void FConcertServerWindowController::OpenSessionTab(const FGuid& SessionId)
+{
+	if (TSharedPtr<FAbstractConcertSessionTab> SessionTab = GetOrRegisterSessionTab(SessionId))
+	{
+		SessionTab->OpenSessionTab();
+	}
+}
+
+TSharedPtr<FAbstractConcertSessionTab> FConcertServerWindowController::GetOrRegisterSessionTab(const FGuid& SessionId)
+{
+	if (TSharedRef<FAbstractConcertSessionTab>* FoundId = RegisteredSessions.Find(SessionId))
+	{
+		return *FoundId;
+	}
+	
+	if (TSharedPtr<IConcertServerSession> Session = ServerInstance->GetConcertServer()->GetLiveSession(SessionId))
+	{
+		const TSharedRef<FLiveConcertSessionTab> SessionTab = MakeShared<FLiveConcertSessionTab>(Session.ToSharedRef(), ServerInstance.ToSharedRef(), RootWindow.ToSharedRef());
+		RegisteredSessions.Add(SessionId, SessionTab);
+		return SessionTab;
+	}
+
+	const bool bIsArchivedSession = ServerInstance->GetConcertServer()->GetArchivedSessionInfo(SessionId).IsSet();
+	if (bIsArchivedSession)
+	{
+		const TSharedRef<FArchivedConcertSessionTab> SessionTab = MakeShared<FArchivedConcertSessionTab>(SessionId, ServerInstance.ToSharedRef(), RootWindow.ToSharedRef());
+		RegisteredSessions.Add(SessionId, SessionTab);
+		return SessionTab;
+	}
+	
+	return nullptr;
+}
+
+void FConcertServerWindowController::InitComponents()
+{
+	const FConcertComponentInitParams Params { ServerInstance.ToSharedRef(), SharedThis(this) };
+	SessionBrowserController->Init(Params);
+}
+
+void FConcertServerWindowController::OnWindowClosed(const TSharedRef<SWindow>& Window)
+{
+	SaveLayout();
+	RootWindow.Reset();
+}
+
+void FConcertServerWindowController::SaveLayout() const
+{
+	FLayoutSaveRestore::SaveToConfig(MultiUserServerLayoutIni, PersistentLayout.ToSharedRef());
+    GConfig->Flush(false, MultiUserServerLayoutIni);
+}
+
+#undef LOCTEXT_NAMESPACE 
\ No newline at end of file
diff --git a/Engine/Plugins/Developer/Concert/ConcertApp/MultiUserServer/Source/MultiUserServer/Private/Widgets/ConcertServerWindowController.h b/Engine/Plugins/Developer/Concert/ConcertApp/MultiUserServer/Source/MultiUserServer/Private/Widgets/ConcertServerWindowController.h
new file mode 100644
--- /dev/null
+++ b/Engine/Plugins/Developer/Concert/ConcertApp/MultiUserServer/Source/MultiUserServer/Private/Widgets/ConcertServerWindowController.h
@@ -0,0 +1,64 @@
+// Copyright Epic Games, Inc. All Rights Reserved.
+
+#pragma once
+
+#include "CoreMinimal.h"
+#include "Framework/Docking/TabManager.h"
+
+class FAbstractConcertSessionTab;
+class FLiveConcertSessionTab;
+class IConcertServerSession;
+class IConcertSyncServer;
+class FConcertServerSessionBrowserController;
+class FOutputLogController;
+class SWindow;
+
+struct FConcertServerWindowInitParams
+{
+	/** The server that the window is supposed to manage */
+	TSharedRef<IConcertSyncServer> Server;
+	
+	/** Config path for server layout ini */
+	FString MultiUserServerLayoutIni;
+	
+	FConcertServerWindowInitParams(TSharedRef<IConcertSyncServer> Server, FString MultiUserServerLayoutIni = FString())
+		: Server(Server)
+		, MultiUserServerLayoutIni(MoveTemp(MultiUserServerLayoutIni))
+	{}
+};
+
+/** Responsible for creating the Slate window for the server. Implements controller in the model-view-controller pattern. */
+class FConcertServerWindowController : public TSharedFromThis<FConcertServerWindowController>
+{
+public:
+	
+	FConcertServerWindowController(const FConcertServerWindowInitParams& Params);
+	void CreateWindow();
+
+	/** Opens or draws attention to the tab for the given live or archived session ID */
+	void OpenSessionTab(const FGuid& SessionId);
+	
+private:
+
+	/** The ini file to use for saving the layout */
+	FString MultiUserServerLayoutIni;
+	/** Holds the current layout for saving later. */
+	TSharedPtr<FTabManager::FLayout> PersistentLayout;
+
+	TSharedPtr<IConcertSyncServer> ServerInstance;
+
+	/** The main window being managed */
+	TSharedPtr<SWindow> RootWindow;
+	TMap<FGuid, TSharedRef<FAbstractConcertSessionTab>> RegisteredSessions;
+	
+	/** Manages the session browser */
+	TSharedPtr<FConcertServerSessionBrowserController> SessionBrowserController;
+	
+	void InitComponents();
+
+	/** Gets the manager for a session tab if the session ID is valid */
+	TSharedPtr<FAbstractConcertSessionTab> GetOrRegisterSessionTab(const FGuid& SessionId);
+
+	void OnWindowClosed(const TSharedRef<SWindow>& Window);
+	void SaveLayout() const;
+};
diff --git a/Engine/Plugins/Developer/Concert/ConcertApp/MultiUserServer/Source/MultiUserServer/Private/Widgets/IConcertComponent.h b/Engine/Plugins/Developer/Concert/ConcertApp/MultiUserServer/Source/MultiUserServer/Private/Widgets/IConcertComponent.h
new file mode 100644
--- /dev/null
+++ b/Engine/Plugins/Developer/Concert/ConcertApp/MultiUserServer/Source/MultiUserServer/Private/Widgets/IConcertComponent.h
@@ -0,0 +1,33 @@
+﻿// Copyright Epic Games, Inc. All Rights Reserved.
+
+#pragma once
+
+#include "CoreMinimal.h"
+
+class FConcertServerWindowController;
+class IConcertSyncServer;
+
+struct FConcertComponentInitParams
+{
+	/** The server being managed */
+	TSharedRef<IConcertSyncServer> Server;
+
+	/** Manages the server slate application */
+	TSharedRef<FConcertServerWindowController> WindowController;
+
+	FConcertComponentInitParams(TSharedRef<IConcertSyncServer> Server, TSharedRef<FConcertServerWindowController> WindowController)
+		: Server(MoveTemp(Server))
+		, WindowController(MoveTemp(WindowController))
+	{}
+};
+
+/** Provides the base interface for elements in the concert server UI. */
+class IConcertComponent
+{
+public:
+
+	virtual ~IConcertComponent() = default;
+
+	/** Initialises the component, e.g. registering tab spawners, etc. */
+	virtual void Init(const FConcertComponentInitParams& Params) {}
+};
diff --git a/Engine/Plugins/Developer/Concert/ConcertApp/MultiUserServer/Source/MultiUserServer/Private/Widgets/SessionTabs/AbstractConcertSessionTab.cpp b/Engine/Plugins/Developer/Concert/ConcertApp/MultiUserServer/Source/MultiUserServer/Private/Widgets/SessionTabs/AbstractConcertSessionTab.cpp
new file mode 100644
--- /dev/null
+++ b/Engine/Plugins/Developer/Concert/ConcertApp/MultiUserServer/Source/MultiUserServer/Private/Widgets/SessionTabs/AbstractConcertSessionTab.cpp
@@ -0,0 +1,65 @@
+﻿// Copyright Epic Games, Inc. All Rights Reserved.
+
+#include "AbstractConcertSessionTab.h"
+
+#include "IConcertServer.h"
+#include "IConcertSyncServer.h"
+#include "Widgets/ConcertServerTabs.h"
+
+#include "Framework/Docking/TabManager.h"
+#include "Widgets/Docking/SDockTab.h"
+
+namespace UE::ConcertServerUI::Private
+{
+	static TOptional<FString> GetSessionName(const TSharedRef<IConcertSyncServer>& SyncServer, const FGuid SessionId)
+	{
+		if (const TSharedPtr<IConcertServerSession> LiveSessionInfo = SyncServer->GetConcertServer()->GetLiveSession(SessionId))
+		{
+			return LiveSessionInfo->GetSessionInfo().SessionName;
+		}
+		if (const TOptional<FConcertSessionInfo> ArchivedSessionInfo = SyncServer->GetConcertServer()->GetArchivedSessionInfo(SessionId))
+		{
+			return ArchivedSessionInfo->SessionName;
+		}
+		return {};
+	}
+}
+
+FAbstractConcertSessionTab::FAbstractConcertSessionTab(TSharedRef<IConcertSyncServer> SyncServer)
+	: SyncServer(MoveTemp(SyncServer))
+{}
+
+void FAbstractConcertSessionTab::OpenSessionTab()
+{
+	const TSharedRef<FGlobalTabmanager>& TabManager = FGlobalTabmanager::Get();
+	const FTabId TabId { *GetTabId() };
+
+	EnsureInitDockTab();
+	if (TabManager->FindExistingLiveTab(TabId))
+	{
+		TabManager->DrawAttention(DockTab.ToSharedRef());
+	}
+	else
+	{
+		const FTabManager::FLastMajorOrNomadTab Search(ConcertServerTabs::GetSessionBrowserTabId());
+		FGlobalTabmanager::Get()->InsertNewDocumentTab(*GetTabId(), Search, DockTab.ToSharedRef());
+
+		OnOpenTab();
+	}
+}
+
+void FAbstractConcertSessionTab::EnsureInitDockTab()
+{
+	if (!DockTab.IsValid())
+	{
+		const TOptional<FString> SessionName = UE::ConcertServerUI::Private::GetSessionName(SyncServer, GetSessionID());
+		check(SessionName);
+		
+		const FText Title = FText::FromString(*SessionName);
+		DockTab = SNew(SDockTab)
+			.Label(Title)
+			.TabRole(MajorTab);
+		
+		CreateDockContent(DockTab.ToSharedRef());
+	}
+}
diff --git a/Engine/Plugins/Developer/Concert/ConcertApp/MultiUserServer/Source/MultiUserServer/Private/Widgets/SessionTabs/AbstractConcertSessionTab.h b/Engine/Plugins/Developer/Concert/ConcertApp/MultiUserServer/Source/MultiUserServer/Private/Widgets/SessionTabs/AbstractConcertSessionTab.h
new file mode 100644
--- /dev/null
+++ b/Engine/Plugins/Developer/Concert/ConcertApp/MultiUserServer/Source/MultiUserServer/Private/Widgets/SessionTabs/AbstractConcertSessionTab.h
@@ -0,0 +1,42 @@
+﻿// Copyright Epic Games, Inc. All Rights Reserved.
+
+#pragma once
+
+#include "CoreMinimal.h"
+
+class IConcertSyncServer;
+class SDockTab;
+class SWindow;
+
+/** Shared functionality for a tab that manages a session */
+class FAbstractConcertSessionTab
+{
+public:
+
+    FAbstractConcertSessionTab(TSharedRef<IConcertSyncServer> SyncServer);
+	virtual ~FAbstractConcertSessionTab() = default;
+	
+	/** Opens or draws attention to the given tab */
+	void OpenSessionTab();
+	
+protected:
+
+	virtual FGuid GetSessionID() const = 0;
+	virtual void CreateDockContent(const TSharedRef<SDockTab>& DockTab) = 0;
+	virtual void OnOpenTab() = 0;
+	
+	/** Generates a tab ID for FTabManager::InsertNewDocumentTab */
+	FString GetTabId() const { return GetSessionID().ToString(); }
+	const TSharedRef<IConcertSyncServer>& GetSyncServer() const { return SyncServer; }
+	
+private:
+
+	/** Used to look up session name */
+	const TSharedRef<IConcertSyncServer> SyncServer;
+	
+	/** The tab containing the UI for InspectedSession */
+	TSharedPtr<SDockTab> DockTab;
+
+	/** Inits DockTab if it has not yet been inited */
+	void EnsureInitDockTab();
+};
diff --git a/Engine/Plugins/Developer/Concert/ConcertApp/MultiUserServer/Source/MultiUserServer/Private/Widgets/SessionTabs/Archived/ArchivedConcertSessionTab.cpp b/Engine/Plugins/Developer/Concert/ConcertApp/MultiUserServer/Source/MultiUserServer/Private/Widgets/SessionTabs/Archived/ArchivedConcertSessionTab.cpp
new file mode 100644
--- /dev/null
+++ b/Engine/Plugins/Developer/Concert/ConcertApp/MultiUserServer/Source/MultiUserServer/Private/Widgets/SessionTabs/Archived/ArchivedConcertSessionTab.cpp
@@ -0,0 +1,31 @@
+﻿// Copyright Epic Games, Inc. All Rights Reserved.
+
+#include "Widgets/SessionTabs/Archived/ArchivedConcertSessionTab.h"
+
+#include "Widgets/SessionTabs/Archived/SConcertArchivedSessionInspector.h"
+#include "Widgets/StatusBar/SConcertStatusBar.h"
+
+FArchivedConcertSessionTab::FArchivedConcertSessionTab(const FGuid& InspectedSessionID, TSharedRef<IConcertSyncServer> SyncServer, TAttribute<TSharedRef<SWindow>> ConstructUnderWindow)
+	: FAbstractConcertSessionTab(SyncServer)
+	, InspectedSessionID(InspectedSessionID)
+	, ConstructUnderWindow(ConstructUnderWindow)
+{}
+
+FGuid FArchivedConcertSessionTab::GetSessionID() const
+{
+	return InspectedSessionID;
+}
+
+void FArchivedConcertSessionTab::CreateDockContent(const TSharedRef<SDockTab>& InDockTab)
+{
+	InDockTab->SetContent(
+		SNew(SConcertArchivedSessionInspector, SConcertArchivedSessionInspector::FRequiredArgs(InDockTab, ConstructUnderWindow.Get()))
+			.StatusBar()
+			[
+				SNew(SConcertStatusBar, *GetTabId())
+			]
+		);
+}
+
+void FArchivedConcertSessionTab::OnOpenTab()
+{}
diff --git a/Engine/Plugins/Developer/Concert/ConcertApp/MultiUserServer/Source/MultiUserServer/Private/Widgets/SessionTabs/Archived/ArchivedConcertSessionTab.h b/Engine/Plugins/Developer/Concert/ConcertApp/MultiUserServer/Source/MultiUserServer/Private/Widgets/SessionTabs/Archived/ArchivedConcertSessionTab.h
new file mode 100644
--- /dev/null
+++ b/Engine/Plugins/Developer/Concert/ConcertApp/MultiUserServer/Source/MultiUserServer/Private/Widgets/SessionTabs/Archived/ArchivedConcertSessionTab.h
@@ -0,0 +1,35 @@
+﻿// Copyright Epic Games, Inc. All Rights Reserved.
+
+#pragma once
+
+#include "CoreMinimal.h"
+#include "SConcertArchivedSessionInspector.h"
+
+#include "Widgets/SessionTabs/AbstractConcertSessionTab.h"
+
+/** Manages the tab for an archived session.  */
+class FArchivedConcertSessionTab : public FAbstractConcertSessionTab
+{
+public:
+
+	FArchivedConcertSessionTab(const FGuid& InspectedSessionID, TSharedRef<IConcertSyncServer> SyncServer, TAttribute<TSharedRef<SWindow>> ConstructUnderWindow);
+
+protected:
+
+	//~ Begin FAbstractConcertSessionTab Interface
+	virtual FGuid GetSessionID() const override;
+	virtual void CreateDockContent(const TSharedRef<SDockTab>& InDockTab) override;
+	virtual void OnOpenTab() override;
+	//~ End FAbstractConcertSessionTab Interface
+	
+private:
+
+	/** The inspected session's ID */
+	FGuid InspectedSessionID;
+
+	/** Used later to obtain the window into which to add the tab */
+	TAttribute<TSharedRef<SWindow>> ConstructUnderWindow;
+
+	/** Displays session */
+	TSharedPtr<SConcertArchivedSessionInspector> Inspector;
+};
diff --git a/Engine/Plugins/Developer/Concert/ConcertApp/MultiUserServer/Source/MultiUserServer/Private/Widgets/SessionTabs/Archived/SConcertArchivedSessionInspector.cpp b/Engine/Plugins/Developer/Concert/ConcertApp/MultiUserServer/Source/MultiUserServer/Private/Widgets/SessionTabs/Archived/SConcertArchivedSessionInspector.cpp
new file mode 100644
--- /dev/null
+++ b/Engine/Plugins/Developer/Concert/ConcertApp/MultiUserServer/Source/MultiUserServer/Private/Widgets/SessionTabs/Archived/SConcertArchivedSessionInspector.cpp
@@ -0,0 +1,89 @@
+﻿// Copyright Epic Games, Inc. All Rights Reserved.
+
+#include "SConcertArchivedSessionInspector.h"
+
+#include "Framework/Docking/TabManager.h"
+#include "Session/History/SSessionHistory.h"
+#include "Widgets/Docking/SDockTab.h"
+#include "Widgets/Layout/SBorder.h"
+#include "Widgets/SBoxPanel.h"
+#include "Widgets/SNullWidget.h"
+
+#define LOCTEXT_NAMESPACE "SConcertArchivedSessionInspector"
+
+const FName SConcertArchivedSessionInspector::HistoryTabId("HistoryTabId");
+
+void SConcertArchivedSessionInspector::Construct(const FArguments& InArgs, const FRequiredArgs& InRequiredArgs)
+{
+	ChildSlot
+	[
+		SNew(SBorder)
+		.BorderImage(FAppStyle::Get().GetBrush("ToolPanel.GroupBorder"))
+		.Padding(FMargin(1.0f, 2.0f))
+		[
+			SNew(SVerticalBox)
+
+			// Toolbar
+			+SVerticalBox::Slot()
+			.AutoHeight()
+			[
+				SNullWidget::NullWidget
+			]
+
+			// Content
+			+SVerticalBox::Slot()
+			.FillHeight(1.f)
+			[
+				SNew(SBorder)
+				.BorderImage(FAppStyle::Get().GetBrush("Brushes.Title"))
+				.Padding(FMargin(0.f, 0.f, 0.f, 5.0f)) // Visually separate status bar from tabs
+				[
+					CreateTabs(InRequiredArgs)
+				]
+			]
+
+			// Status bar
+			+SVerticalBox::Slot()
+			.AutoHeight()
+			.HAlign(HAlign_Fill)
+			[
+				InArgs._StatusBar.Widget
+			]
+		]
+	];
+}
+
+TSharedRef<SWidget> SConcertArchivedSessionInspector::CreateTabs(const FRequiredArgs& RequiredArgs)
+{
+	TabManager = FGlobalTabmanager::Get()->NewTabManager(RequiredArgs.ConstructUnderMajorTab);
+	
+	TabManager->RegisterTabSpawner(HistoryTabId, FOnSpawnTab::CreateSP(this, &SConcertArchivedSessionInspector::SpawnActivityHistory))
+		.SetDisplayName(LOCTEXT("ActivityHistoryLabel", "History"));
+
+	const TSharedRef<FTabManager::FLayout> Layout = FTabManager::NewLayout("ConcertArchivedSessionLayout_v0.1")
+		->AddArea
+		(
+			FTabManager::NewPrimaryArea()
+				->SetOrientation(Orient_Horizontal)
+				->Split
+				(
+					FTabManager::NewStack()
+					->SetSizeCoefficient(0.5f)
+					->AddTab(HistoryTabId, ETabState::OpenedTab)
+				)
+		);
+
+	return TabManager->RestoreFrom(Layout, RequiredArgs.ConstructUnderWindow).ToSharedRef();
+}
+
+TSharedRef<SDockTab> SConcertArchivedSessionInspector::SpawnActivityHistory(const FSpawnTabArgs& Args)
+{
+	return SNew(SDockTab)
+		.Label(LOCTEXT("ActivityHistoryLabel", "History"))
+		.TabRole(PanelTab)
+		[
+			SNullWidget::NullWidget
+		]; 
+}
+
+#undef LOCTEXT_NAMESPACE
diff --git a/Engine/Plugins/Developer/Concert/ConcertApp/MultiUserServer/Source/MultiUserServer/Private/Widgets/SessionTabs/Archived/SConcertArchivedSessionInspector.h b/Engine/Plugins/Developer/Concert/ConcertApp/MultiUserServer/Source/MultiUserServer/Private/Widgets/SessionTabs/Archived/SConcertArchivedSessionInspector.h
new file mode 100644
--- /dev/null
+++ b/Engine/Plugins/Developer/Concert/ConcertApp/MultiUserServer/Source/MultiUserServer/Private/Widgets/SessionTabs/Archived/SConcertArchivedSessionInspector.h
@@ -0,0 +1,42 @@
+﻿// Copyright Epic Games, Inc. All Rights Reserved.
+
+#pragma once
+
+#include "CoreMinimal.h"
+#include "Widgets/DeclarativeSyntaxSupport.h"
+#include "Widgets/Docking/SDockTab.h"
+#include "Widgets/SCompoundWidget.h"
+
+/** Designed as content for a tab. Displays information about an archived session. */
+class SConcertArchivedSessionInspector : public SCompoundWidget
+{
+public:
+
+	struct FRequiredArgs
+	{
+		TSharedRef<SDockTab> ConstructUnderMajorTab;
+		TSharedRef<SWindow> ConstructUnderWindow;
+
+		FRequiredArgs(TSharedRef<SDockTab> ConstructUnderMajorTab, TSharedRef<SWindow> ConstructUnderWindow)
+			: ConstructUnderMajorTab(MoveTemp(ConstructUnderMajorTab))
+			, ConstructUnderWindow(MoveTemp(ConstructUnderWindow))
+		{}
+	};
+	
+	static const FName HistoryTabId;
+
+	SLATE_BEGIN_ARGS(SConcertArchivedSessionInspector)
+	{}
+		SLATE_NAMED_SLOT(FArguments, StatusBar)
+	SLATE_END_ARGS()
+
+	void Construct(const FArguments& InArgs, const FRequiredArgs& InRequiredArgs);
+
+private:
+	
+	/** Holds the tab manager that manages the front-end's tabs. */
+	TSharedPtr<FTabManager> TabManager;
+	
+	TSharedRef<SWidget> CreateTabs(const FRequiredArgs& RequiredArgs);
+	TSharedRef<SDockTab> SpawnActivityHistory(const FSpawnTabArgs& Args);
+};
diff --git a/Engine/Plugins/Developer/Concert/ConcertApp/MultiUserServer/Source/MultiUserServer/Private/Widgets/SessionTabs/Live/LiveConcertSessionTab.cpp b/Engine/Plugins/Developer/Concert/ConcertApp/MultiUserServer/Source/MultiUserServer/Private/Widgets/SessionTabs/Live/LiveConcertSessionTab.cpp
new file mode 100644
--- /dev/null
+++ b/Engine/Plugins/Developer/Concert/ConcertApp/MultiUserServer/Source/MultiUserServer/Private/Widgets/SessionTabs/Live/LiveConcertSessionTab.cpp
@@ -0,0 +1,48 @@
+﻿// Copyright Epic Games, Inc. All Rights Reserved.
+
+#include "LiveConcertSessionTab.h"
+
+#include "IConcertSession.h"
+#include "PackageViewer/ConcertSessionPackageViewerController.h"
+#include "SConcertSessionInspector.h"
+#include "Widgets/StatusBar/SConcertStatusBar.h"
+
+#include "Widgets/Docking/SDockTab.h"
+
+FLiveConcertSessionTab::FLiveConcertSessionTab(TSharedRef<IConcertServerSession> InspectedSession, TSharedRef<IConcertSyncServer> SyncServer, TAttribute<TSharedRef<SWindow>> ConstructUnderWindow)
+	: FAbstractConcertSessionTab(SyncServer)
+	, InspectedSession(MoveTemp(InspectedSession))
+	, ConstructUnderWindow(MoveTemp(ConstructUnderWindow))
+	, SessionHistoryController(MakeShared<FServerSessionHistoryController>(InspectedSession, SyncServer))
+	, PackageViewerController(MakeShared<FConcertSessionPackageViewerController>(InspectedSession, SyncServer))
+{}
+
+FGuid FLiveConcertSessionTab::GetSessionID() const
+{
+	return InspectedSession->GetSessionInfo().SessionId;
+}
+
+void FLiveConcertSessionTab::CreateDockContent(const TSharedRef<SDockTab>& InDockTab)
+{
+	const SConcertSessionInspector::FRequiredArgs WidgetArgs
+	{
+		InDockTab,
+		ConstructUnderWindow.Get(),
+		SessionHistoryController->GetSessionHistory(),
+		PackageViewerController->GetPackageViewer()
+	};
+	InDockTab->SetContent(
+		SNew(SConcertSessionInspector, WidgetArgs)
+			.StatusBar()
+			[
+				SNew(SConcertStatusBar, *GetTabId())
+			]
+		);
+}
+
+void FLiveConcertSessionTab::OnOpenTab()
+{
+	SessionHistoryController->ReloadActivities();
+	PackageViewerController->ReloadActivities();
+}
+
diff --git a/Engine/Plugins/Developer/Concert/ConcertApp/MultiUserServer/Source/MultiUserServer/Private/Widgets/SessionTabs/Live/LiveConcertSessionTab.h b/Engine/Plugins/Developer/Concert/ConcertApp/MultiUserServer/Source/MultiUserServer/Private/Widgets/SessionTabs/Live/LiveConcertSessionTab.h
new file mode 100644
--- /dev/null
+++ b/Engine/Plugins/Developer/Concert/ConcertApp/MultiUserServer/Source/MultiUserServer/Private/Widgets/SessionTabs/Live/LiveConcertSessionTab.h
@@ -0,0 +1,47 @@
+﻿// Copyright Epic Games, Inc. All Rights Reserved.
+
+#pragma once
+
+#include "CoreMinimal.h"
+#include "ServerSessionHistoryController.h"
+#include "Widgets/SessionTabs/AbstractConcertSessionTab.h"
+#include "Types/SlateAttribute.h"
+
+class FConcertSessionPackageViewerController;
+class IConcertSyncServer;
+class FServerSessionHistoryController;
+class IConcertServerSession;
+class SDockTab;
+class SWindow;
+
+/**
+ * Manages the tab that contains the UI for a session.
+ * It has access to the controllers and views needed for displaying a session.
+ */
+class FLiveConcertSessionTab : public FAbstractConcertSessionTab
+{
+public:
+
+	FLiveConcertSessionTab(TSharedRef<IConcertServerSession> InspectedSession, TSharedRef<IConcertSyncServer> SyncServer, TAttribute<TSharedRef<SWindow>> ConstructUnderWindow);
+
+protected:
+
+	//~ Begin FAbstractConcertSessionTab Interface
+	virtual FGuid GetSessionID() const override;
+	virtual void CreateDockContent(const TSharedRef<SDockTab>& InDockTab) override;
+	virtual void OnOpenTab() override;
+	//~ End FAbstractConcertSessionTab Interface
+	
+private:
+	
+	/** The session being inspected */
+	const TSharedRef<IConcertServerSession> InspectedSession;
+
+	/** Used to determine owning window during construction */
+	TAttribute<TSharedRef<SWindow>> ConstructUnderWindow;
+
+	/** Manages the session history widget */
+	const TSharedRef<FServerSessionHistoryController> SessionHistoryController;
+	/** Manages the package viewer widget */
+	const TSharedRef<FConcertSessionPackageViewerController> PackageViewerController;
+};
diff --git a/Engine/Plugins/Developer/Concert/ConcertApp/MultiUserServer/Source/MultiUserServer/Private/Widgets/SessionTabs/Live/PackageViewer/ConcertSessionPackageViewerController.cpp b/Engine/Plugins/Developer/Concert/ConcertApp/MultiUserServer/Source/MultiUserServer/Private/Widgets/SessionTabs/Live/PackageViewer/ConcertSessionPackageViewerController.cpp
new file mode 100644
--- /dev/null
+++ b/Engine/Plugins/Developer/Concert/ConcertApp/MultiUserServer/Source/MultiUserServer/Private/Widgets/SessionTabs/Live/PackageViewer/ConcertSessionPackageViewerController.cpp
@@ -0,0 +1,104 @@
+﻿// Copyright Epic Games, Inc. All Rights Reserved.
+
+#include "ConcertSessionPackageViewerController.h"
+
+#include "ConcertSyncSessionDatabase.h"
+#include "IConcertSyncServer.h"
+#include "SConcertSessionPackageViewer.h"
+#include "Concert/Private/ConcertServerSession.h"
+
+FConcertSessionPackageViewerController::FConcertSessionPackageViewerController(TSharedRef<IConcertServerSession> InspectedSession, TSharedRef<IConcertSyncServer> SyncServer)
+	: InspectedSession(MoveTemp(InspectedSession))
+	, SyncServer(MoveTemp(SyncServer))
+	, PackageViewer(MakePackageViewer())
+{
+	ReloadActivities();
+	
+	if (const TOptional<FConcertSyncSessionDatabaseNonNullPtr> Database = SyncServer->GetLiveSessionDatabase(InspectedSession->GetId()))
+	{
+		Database->OnActivityProduced().AddRaw(this, &FConcertSessionPackageViewerController::OnSessionProduced);
+	}
+}
+
+FConcertSessionPackageViewerController::~FConcertSessionPackageViewerController()
+{
+	if (const TOptional<FConcertSyncSessionDatabaseNonNullPtr> Database = SyncServer->GetLiveSessionDatabase(InspectedSession->GetId()))
+	{
+		Database->OnActivityProduced().RemoveAll(this);
+	}
+}
+
+void FConcertSessionPackageViewerController::ReloadActivities() const
+{
+	PackageViewer->ResetActivityList();
+	
+	if (const TOptional<FConcertSyncSessionDatabaseNonNullPtr> Database = SyncServer->GetLiveSessionDatabase(InspectedSession->GetId()))
+	{
+		struct FPackageActivity
+		{
+			FConcertSessionActivity Activity;
+			FConcertSyncPackageEventData PackageData;
+		};
+		
+		TMap<FName, FPackageActivity> LatestPackageActivities;
+		Database->EnumeratePackageActivities([this, &LatestPackageActivities](FConcertSyncActivity&& BasePart, FConcertSyncPackageEventData& EventData)
+		{
+			FStructOnScope ActivitySummary;
+			if (BasePart.EventSummary.GetPayload(ActivitySummary))
+			{
+				FPackageActivity Activity{ { BasePart, ActivitySummary }, EventData };
+				const bool bWasRenamed = Activity.PackageData.MetaData.PackageInfo.PackageUpdateType == EConcertPackageUpdateType::Renamed;
+				if (bWasRenamed)
+				{
+					LatestPackageActivities.Remove(EventData.MetaData.PackageInfo.PackageName);
+					LatestPackageActivities.Emplace(EventData.MetaData.PackageInfo.NewPackageName, MoveTemp(Activity));
+				}
+				else
+				{
+					LatestPackageActivities.Emplace(EventData.MetaData.PackageInfo.PackageName, MoveTemp(Activity));
+				}
+			}
+			return EBreakBehavior::Continue;
+		});
+
+		for (auto ActivityIt = LatestPackageActivities.CreateIterator(); ActivityIt; ++ActivityIt)
+		{
+			PackageViewer->AppendActivity(MoveTemp(ActivityIt->Value.Activity));
+		}
+	}
+}
+
+TSharedRef<SConcertSessionPackageViewer> FConcertSessionPackageViewerController::MakePackageViewer() const
+{
+	return SNew(SConcertSessionPackageViewer)
+		.GetClientInfo_Raw(this, &FConcertSessionPackageViewerController::GetClientInfo)
+		.GetPackageEvent_Raw(this, &FConcertSessionPackageViewerController::GetPackageEvent);
+}
+
+TOptional<FConcertClientInfo> FConcertSessionPackageViewerController::GetClientInfo(FGuid ClientId) const
+{
+	const TOptional<FConcertSyncSessionDatabaseNonNullPtr> Database = SyncServer->GetLiveSessionDatabase(InspectedSession->GetId());
+	FConcertSyncEndpointData Result;
+	if (Database && Database->GetEndpoint(ClientId, Result))
+	{
+		return Result.ClientInfo;
+	}
+	return {};
+}
+
+bool FConcertSessionPackageViewerController::GetPackageEvent(const FConcertSessionActivity& Activity, FConcertSyncPackageEventMetaData& OutPackageEvent) const
+{
+	if (const TOptional<FConcertSyncSessionDatabaseNonNullPtr> Database = SyncServer->GetLiveSessionDatabase(InspectedSession->GetId()))
+	{
+		return Database->GetPackageEventMetaData(Activity.Activity.EventId, OutPackageEvent.PackageRevision, OutPackageEvent.PackageInfo);
+	}
+	return false;
+}
+
+void FConcertSessionPackageViewerController::OnSessionProduced(const FConcertSyncActivity& ProducedActivity) const
+{
+	if (ProducedActivity.EventType == EConcertSyncActivityEventType::Package)
+	{
+		ReloadActivities();
+	}
+}
diff --git a/Engine/Plugins/Developer/Concert/ConcertApp/MultiUserServer/Source/MultiUserServer/Private/Widgets/SessionTabs/Live/PackageViewer/ConcertSessionPackageViewerController.h b/Engine/Plugins/Developer/Concert/ConcertApp/MultiUserServer/Source/MultiUserServer/Private/Widgets/SessionTabs/Live/PackageViewer/ConcertSessionPackageViewerController.h
new file mode 100644
--- /dev/null
+++ b/Engine/Plugins/Developer/Concert/ConcertApp/MultiUserServer/Source/MultiUserServer/Private/Widgets/SessionTabs/Live/PackageViewer/ConcertSessionPackageViewerController.h
@@ -0,0 +1,39 @@
+﻿// Copyright Epic Games, Inc. All Rights Reserved.
+
+#pragma once
+
+#include "CoreMinimal.h"
+
+struct FConcertSyncActivity;
+struct FConcertClientInfo;
+struct FConcertSyncPackageEventMetaData;
+struct FConcertSessionActivity;
+class IConcertSyncServer;
+class IConcertServerSession;
+class SConcertSessionPackageViewer;
+
+/** Manages SConcertSessionPackageViewer for a server. */
+class FConcertSessionPackageViewerController
+{
+public:
+
+	FConcertSessionPackageViewerController(TSharedRef<IConcertServerSession> InspectedSession, TSharedRef<IConcertSyncServer> SyncServer);
+	~FConcertSessionPackageViewerController();
+	
+	void ReloadActivities() const;
+	TSharedRef<SConcertSessionPackageViewer> GetPackageViewer() const { return PackageViewer; }
+
+private:
+	
+	TSharedRef<IConcertServerSession> InspectedSession;
+	TSharedRef<IConcertSyncServer> SyncServer;
+	
+	/** The widget being managed */
+	TSharedRef<SConcertSessionPackageViewer> PackageViewer;
+	
+	TSharedRef<SConcertSessionPackageViewer> MakePackageViewer() const;
+	
+	TOptional<FConcertClientInfo> GetClientInfo(FGuid ClientId) const;
+	bool GetPackageEvent(const FConcertSessionActivity& Activity, FConcertSyncPackageEventMetaData& OutPackageEvent) const;
+	void OnSessionProduced(const FConcertSyncActivity& ProducedActivity) const;
+};
diff --git a/Engine/Plugins/Developer/Concert/ConcertApp/MultiUserServer/Source/MultiUserServer/Private/Widgets/SessionTabs/Live/PackageViewer/SConcertSessionPackageViewer.cpp b/Engine/Plugins/Developer/Concert/ConcertApp/MultiUserServer/Source/MultiUserServer/Private/Widgets/SessionTabs/Live/PackageViewer/SConcertSessionPackageViewer.cpp
new file mode 100644
--- /dev/null
+++ b/Engine/Plugins/Developer/Concert/ConcertApp/MultiUserServer/Source/MultiUserServer/Private/Widgets/SessionTabs/Live/PackageViewer/SConcertSessionPackageViewer.cpp
@@ -0,0 +1,97 @@
+﻿// Copyright Epic Games, Inc. All Rights Reserved.
+
+#include "SConcertSessionPackageViewer.h"
+
+#include "Session/Activity/PredefinedActivityColumns.h"
+#include "Session/Activity/SConcertSessionActivities.h"
+#include "Widgets/Input/SSearchBox.h"
+
+#define LOCTEXT_NAMESPACE "UnrealMultiUserUI"
+
+void SConcertSessionPackageViewer::Construct(const FArguments& InArgs)
+{
+	using namespace UE::ConcertSharedSlate;
+	
+	ActivityListViewOptions = MakeShared<FConcertSessionActivitiesOptions>();
+	ActivityListViewOptions->bEnableConnectionActivityFiltering = false;
+	ActivityListViewOptions->bEnableLockActivityFiltering = false;
+	ActivityListViewOptions->bEnablePackageActivityFiltering = false;
+	ActivityListViewOptions->bEnableTransactionActivityFiltering = false;
+
+	SAssignNew(ActivityListView, SConcertSessionActivities)
+		.OnGetPackageEvent(InArgs._GetPackageEvent)
+		.OnFetchActivities(InArgs._FetchInitialActivities)
+		.OnMapActivityToClient(InArgs._GetClientInfo)
+		.HighlightText(this, &SConcertSessionPackageViewer::HighlightSearchedText)
+		.Columns({
+			// TODO: Create Icon (depending on EConcertPackageUpdateType), Size, and Version
+		})
+		.TimeFormat(ActivityListViewOptions.Get(), &FConcertSessionActivitiesOptions::GetTimeFormat)
+		.ConnectionActivitiesVisibility(EVisibility::Collapsed)
+		.LockActivitiesVisibility(EVisibility::Collapsed)
+		.PackageActivitiesVisibility(EVisibility::Visible)
+		.TransactionActivitiesVisibility(EVisibility::Collapsed)
+		.DetailsAreaVisibility(EVisibility::Collapsed)
+		.IsAutoScrollEnabled(true);
+
+	ChildSlot
+	[
+		SNew(SVerticalBox)
+		+SVerticalBox::Slot()
+		.AutoHeight()
+		.Padding(1, 1)
+		[
+			SAssignNew(SearchBox, SSearchBox)
+			.HintText(LOCTEXT("SearchHint", "Search..."))
+			.OnTextChanged(this, &SConcertSessionPackageViewer::OnSearchTextChanged)
+			.OnTextCommitted(this, &SConcertSessionPackageViewer::OnSearchTextCommitted)
+			.DelayChangeNotificationsWhileTyping(true)
+		]
+
+		+SVerticalBox::Slot()
+		[
+			ActivityListView.ToSharedRef()
+		]
+
+		+SVerticalBox::Slot()
+		.AutoHeight()
+		[
+			ActivityListViewOptions->MakeStatusBar(
+				TAttribute<int32>(ActivityListView.Get(), &SConcertSessionActivities::GetTotalActivityNum),
+				TAttribute<int32>(ActivityListView.Get(), &SConcertSessionActivities::GetDisplayedActivityNum)
+				)
+		]
+	];
+}
+
+void SConcertSessionPackageViewer::ResetActivityList()
+{
+	ActivityListView->ResetActivityList();
+}
+
+void SConcertSessionPackageViewer::AppendActivity(FConcertSessionActivity Activity)
+{
+	const TSharedRef<FConcertSessionActivity> NewActivity = MakeShared<FConcertSessionActivity>(MoveTemp(Activity));
+	ActivityListView->Append(NewActivity);
+}
+
+void SConcertSessionPackageViewer::OnSearchTextChanged(const FText& InSearchText)
+{
+	SearchedText = InSearchText;
+	SearchBox->SetError(ActivityListView->UpdateTextFilter(InSearchText));
+}
+
+void SConcertSessionPackageViewer::OnSearchTextCommitted(const FText& InFilterText, ETextCommit::Type CommitType)
+{
+	if (!InFilterText.EqualTo(SearchedText))
+	{
+		OnSearchTextChanged(InFilterText);
+	}
+}
+
+FText SConcertSessionPackageViewer::HighlightSearchedText() const
+{
+	return SearchedText;
+}
+
+#undef LOCTEXT_NAMESPACE
diff --git a/Engine/Plugins/Developer/Concert/ConcertApp/MultiUserServer/Source/MultiUserServer/Private/Widgets/SessionTabs/Live/PackageViewer/SConcertSessionPackageViewer.h b/Engine/Plugins/Developer/Concert/ConcertApp/MultiUserServer/Source/MultiUserServer/Private/Widgets/SessionTabs/Live/PackageViewer/SConcertSessionPackageViewer.h
new file mode 100644
--- /dev/null
+++ b/Engine/Plugins/Developer/Concert/ConcertApp/MultiUserServer/Source/MultiUserServer/Private/Widgets/SessionTabs/Live/PackageViewer/SConcertSessionPackageViewer.h
@@ -0,0 +1,54 @@
+﻿// Copyright Epic Games, Inc. All Rights Reserved.
+
+#pragma once
+
+#include "CoreMinimal.h"
+#include "Session/Activity/SConcertSessionActivities.h"
+#include "Widgets/SCompoundWidget.h"
+#include "Widgets/Input/SSearchBox.h"
+#include "Widgets/DeclarativeSyntaxSupport.h"
+
+class FConcertSessionActivitiesOptions;
+
+/** Displays package activities that happened in a concert session. */
+class SConcertSessionPackageViewer : public SCompoundWidget
+{
+public:
+
+	struct FConcertPackageActivity
+	{
+		FConcertSyncActivity Base;
+		FConcertSyncPackageEventData PackageEvent;
+	};
+	
+	SLATE_BEGIN_ARGS(SConcertSessionPackageViewer) {}
+		SLATE_EVENT(SConcertSessionActivities::FGetPackageEvent, GetPackageEvent)
+		SLATE_EVENT(SConcertSessionActivities::FGetActivityClientInfoFunc, GetClientInfo)
+		SLATE_EVENT(SConcertSessionActivities::FFetchActivitiesFunc, FetchInitialActivities)
+	SLATE_END_ARGS()
+
+	void Construct(const FArguments& InArgs);
+
+	void ResetActivityList();
+	void AppendActivity(FConcertSessionActivity Activity);
+
+private:
+
+	TSharedPtr<SConcertSessionActivities> ActivityListView;
+	
+	/** Controls the activity list view options */
+	TSharedPtr<FConcertSessionActivitiesOptions> ActivityListViewOptions;
+
+	/** The widget used to enter the text to search. */
+	TSharedPtr<SSearchBox> SearchBox;
+
+	/** The searched text to highlight. */
+	FText SearchedText;
+
+	/** Invoked when the text in the search box widget changes. */
+	void OnSearchTextChanged(const FText& InSearchText);
+	/** Invoked when the text in the search box widget is committed. */
+	void OnSearchTextCommitted(const FText& InFilterText, ETextCommit::Type CommitType);
+	/** Returns the text to highlight when the search bar has a text set. */
+	FText HighlightSearchedText() const;
+};
diff --git a/Engine/Plugins/Developer/Concert/ConcertApp/MultiUserServer/Source/MultiUserServer/Private/Widgets/SessionTabs/Live/SConcertSessionInspector.cpp b/Engine/Plugins/Developer/Concert/ConcertApp/MultiUserServer/Source/MultiUserServer/Private/Widgets/SessionTabs/Live/SConcertSessionInspector.cpp
new file mode 100644
--- /dev/null
+++ b/Engine/Plugins/Developer/Concert/ConcertApp/MultiUserServer/Source/MultiUserServer/Private/Widgets/SessionTabs/Live/SConcertSessionInspector.cpp
@@ -0,0 +1,127 @@
+﻿// Copyright Epic Games, Inc. All Rights Reserved.
+
+#include "SConcertSessionInspector.h"
+
+#include "Framework/Docking/TabManager.h"
+#include "Session/History/SSessionHistory.h"
+#include "Widgets/Docking/SDockTab.h"
+#include "Widgets/Layout/SBorder.h"
+#include "Widgets/SBoxPanel.h"
+
+#define LOCTEXT_NAMESPACE "SConcertSessionInspector"
+
+const FName SConcertSessionInspector::HistoryTabId("HistoryTabId");
+const FName SConcertSessionInspector::SessionContentTabId("SessionContentTabId");
+const FName SConcertSessionInspector::ConnectionMonitorTabId("ConnectionMonitorTabId");
+
+void SConcertSessionInspector::Construct(const FArguments& InArgs, const FRequiredArgs& RequiredArgs)
+{
+	ChildSlot
+	[
+		SNew(SBorder)
+		.BorderImage(FAppStyle::Get().GetBrush("ToolPanel.GroupBorder"))
+		.Padding(FMargin(1.0f, 2.0f))
+		[
+			SNew(SVerticalBox)
+
+			// Content
+			+SVerticalBox::Slot()
+			.FillHeight(1.f)
+			[
+				SNew(SBorder)
+				.BorderImage(FAppStyle::Get().GetBrush("Brushes.Title"))
+				.Padding(FMargin(0.f, 0.f, 0.f, 5.0f)) // Visually separate status bar from tabs
+				[
+					CreateTabs(InArgs, RequiredArgs)
+				]
+			]
+
+			// Status bar
+			+SVerticalBox::Slot()
+			.AutoHeight()
+			.HAlign(HAlign_Fill)
+			[
+				InArgs._StatusBar.Widget
+			]
+		]
+	];
+}
+
+TSharedRef<SWidget> SConcertSessionInspector::CreateTabs(const FArguments& InArgs, const FRequiredArgs& RequiredArgs)
+{
+	TabManager = FGlobalTabmanager::Get()->NewTabManager(RequiredArgs.ConstructUnderMajorTab);
+	
+	TabManager->RegisterTabSpawner(HistoryTabId, FOnSpawnTab::CreateSP(this, &SConcertSessionInspector::SpawnActivityHistory, RequiredArgs.SessionHistory))
+		.SetDisplayName(LOCTEXT("ActivityHistoryLabel", "History"));
+	
+	TabManager->RegisterTabSpawner(SessionContentTabId, FOnSpawnTab::CreateSP(this, &SConcertSessionInspector::SpawnSessionContent, RequiredArgs.PackageViewer))
+		.SetDisplayName(LOCTEXT("SessionContentLabel", "Session Content"));
+	
+	TabManager->RegisterTabSpawner(ConnectionMonitorTabId, FOnSpawnTab::CreateSP(this, &SConcertSessionInspector::SpawnConnectionMonitor))
+		.SetDisplayName(LOCTEXT("ConnectionMonitorLabel", "Connection Monitor"));
+
+	const TSharedRef<FTabManager::FLayout> Layout = FTabManager::NewLayout("ConcertSessionLayout_v0.3")
+		->AddArea
+		(
+			FTabManager::NewPrimaryArea()
+				->SetOrientation(Orient_Horizontal)
+				->Split
+				(
+					FTabManager::NewStack()
+					->SetSizeCoefficient(0.5f)
+					->AddTab(HistoryTabId, ETabState::OpenedTab)
+				)
+				->Split
+				(
+					FTabManager::NewSplitter()
+					->SetSizeCoefficient(0.5f)
+					->SetOrientation(Orient_Vertical)
+					->Split
+					(
+						FTabManager::NewStack()
+						->SetSizeCoefficient(0.5f)
+						->AddTab(SessionContentTabId, ETabState::OpenedTab)
+					)
+					->Split
+					(
+						FTabManager::NewStack()
+						->SetSizeCoefficient(0.5f)
+						->AddTab(ConnectionMonitorTabId, ETabState::OpenedTab)
+					)
+				)
+		);
+
+	return TabManager->RestoreFrom(Layout, RequiredArgs.ConstructUnderWindow).ToSharedRef();
+}
+
+TSharedRef<SDockTab> SConcertSessionInspector::SpawnActivityHistory(const FSpawnTabArgs& Args, TSharedRef<SSessionHistory> SessionHistory)
+{
+	return SNew(SDockTab)
+		.Label(LOCTEXT("ActivityHistoryLabel", "History"))
+		.TabRole(PanelTab)
+		[
+			SessionHistory
+		];
+}
+
+TSharedRef<SDockTab> SConcertSessionInspector::SpawnSessionContent(const FSpawnTabArgs& Args, TSharedRef<SConcertSessionPackageViewer> PackageViewer)
+{
+	return SNew(SDockTab)
+		.Label(LOCTEXT("SessionContentLabel", "Session Content"))
+		.TabRole(PanelTab)
+		[
+			PackageViewer
+		];
+}
+
+TSharedRef<SDockTab> SConcertSessionInspector::SpawnConnectionMonitor(const FSpawnTabArgs& Args)
+{
+	return SNew(SDockTab)
+		.Label(LOCTEXT("ConnectionMonitorLabel", "Connection Monitor"))
+		.TabRole(PanelTab)
+		[
+			SNullWidget::NullWidget
+		];
+}
+
+#undef LOCTEXT_NAMESPACE
diff --git a/Engine/Plugins/Developer/Concert/ConcertApp/MultiUserServer/Source/MultiUserServer/Private/Widgets/SessionTabs/Live/SConcertSessionInspector.h b/Engine/Plugins/Developer/Concert/ConcertApp/MultiUserServer/Source/MultiUserServer/Private/Widgets/SessionTabs/Live/SConcertSessionInspector.h
new file mode 100644
--- /dev/null
+++ b/Engine/Plugins/Developer/Concert/ConcertApp/MultiUserServer/Source/MultiUserServer/Private/Widgets/SessionTabs/Live/SConcertSessionInspector.h
@@ -0,0 +1,64 @@
+﻿// Copyright Epic Games, Inc. All Rights Reserved.
+
+#pragma once
+
+#include "CoreMinimal.h"
+#include "Session/History/SSessionHistory.h"
+#include "PackageViewer/SConcertSessionPackageViewer.h"
+#include "Widgets/SCompoundWidget.h"
+#include "Widgets/DeclarativeSyntaxSupport.h"
+
+class FTabManager;
+class FSpawnTabArgs;
+class SDockTab;
+class SWindow;
+
+/**
+ * Designed to be the content of a tab showing:
+ *  - activity history (transactions stored on the server as well as who made those transactions)
+ *  - session content (list of session data saved during a Multi-user session)
+ *  - connection monitor (details about the connected clients on the given session and network info)
+ * Implements view in model-view-controller pattern.
+ */
+class SConcertSessionInspector : public SCompoundWidget
+{
+public:
+
+	struct FRequiredArgs
+	{
+		TSharedRef<SDockTab> ConstructUnderMajorTab;
+		TSharedRef<SWindow> ConstructUnderWindow;
+		TSharedRef<SSessionHistory> SessionHistory;
+		TSharedRef<SConcertSessionPackageViewer> PackageViewer;
+
+		FRequiredArgs(TSharedRef<SDockTab> ConstructUnderMajorTab, TSharedRef<SWindow> ConstructUnderWindow,
+				TSharedRef<SSessionHistory> SessionHistoryController,
+				TSharedRef<SConcertSessionPackageViewer> PackageViewerController)
+			: ConstructUnderMajorTab(MoveTemp(ConstructUnderMajorTab))
+			, ConstructUnderWindow(MoveTemp(ConstructUnderWindow))
+			, SessionHistory(MoveTemp(SessionHistoryController))
+			, PackageViewer(MoveTemp(PackageViewerController))
+		{}
+	};
+
+	static const FName HistoryTabId;
+	static const FName SessionContentTabId;
+	static const FName ConnectionMonitorTabId;
+
+	SLATE_BEGIN_ARGS(SConcertSessionInspector) {}
+		SLATE_NAMED_SLOT(FArguments, StatusBar)
+	SLATE_END_ARGS()
+
+	void Construct(const FArguments& InArgs, const FRequiredArgs& RequiredArgs);
+
+private:
+	
+	/** Holds the tab manager that manages the front-end's tabs. */
+	TSharedPtr<FTabManager> TabManager;
+
+	// Spawn tabs
+	TSharedRef<SWidget> CreateTabs(const FArguments& InArgs, const FRequiredArgs& RequiredArgs);
+	TSharedRef<SDockTab> SpawnActivityHistory(const FSpawnTabArgs& Args, TSharedRef<SSessionHistory> SessionHistory);
+	TSharedRef<SDockTab> SpawnSessionContent(const FSpawnTabArgs& Args, TSharedRef<SConcertSessionPackageViewer> PackageViewer);
+	TSharedRef<SDockTab> SpawnConnectionMonitor(const FSpawnTabArgs& Args);
+};
diff --git a/Engine/Plugins/Developer/Concert/ConcertApp/MultiUserServer/Source/MultiUserServer/Private/Widgets/SessionTabs/Live/ServerSessionHistoryController.cpp b/Engine/Plugins/Developer/Concert/ConcertApp/MultiUserServer/Source/MultiUserServer/Private/Widgets/SessionTabs/Live/ServerSessionHistoryController.cpp
new file mode 100644
--- /dev/null
+++ b/Engine/Plugins/Developer/Concert/ConcertApp/MultiUserServer/Source/MultiUserServer/Private/Widgets/SessionTabs/Live/ServerSessionHistoryController.cpp
@@ -0,0 +1,92 @@
+﻿// Copyright Epic Games, Inc. All Rights Reserved.
+
+#include "ServerSessionHistoryController.h"
+
+#include "ConcertSyncSessionDatabase.h"
+#include "IConcertSession.h"
+#include "IConcertSyncServer.h"
+
+FServerSessionHistoryController::FServerSessionHistoryController(TSharedRef<IConcertServerSession> InspectedSession, TSharedRef<IConcertSyncServer> SyncServer)
+	: InspectedSession(MoveTemp(InspectedSession))
+	, SyncServer(MoveTemp(SyncServer))
+{
+	ReloadActivities();
+
+	if (const TOptional<FConcertSyncSessionDatabaseNonNullPtr> Database = SyncServer->GetLiveSessionDatabase(InspectedSession->GetId()))
+	{
+		Database->OnActivityProduced().AddRaw(this, &FServerSessionHistoryController::OnSessionProduced);
+	}
+}
+
+FServerSessionHistoryController::~FServerSessionHistoryController()
+{
+	if (const TOptional<FConcertSyncSessionDatabaseNonNullPtr> Database = SyncServer->GetLiveSessionDatabase(InspectedSession->GetId()))
+	{
+		Database->OnActivityProduced().RemoveAll(this);
+	}
+}
+
+void FServerSessionHistoryController::GetActivities(int64 MaximumNumberOfActivities, TMap<FGuid, FConcertClientInfo>& OutEndpointClientInfoMap, TArray<FConcertSessionActivity>& OutFetchedActivities) const
+{
+	if (const TOptional<FConcertSyncSessionDatabaseNonNullPtr> Database = SyncServer->GetLiveSessionDatabase(InspectedSession->GetId()))
+	{
+		OutEndpointClientInfoMap.Reset();
+		OutFetchedActivities.Reset();
+
+		int64 LastActivityId = INDEX_NONE;
+		Database->GetActivityMaxId(LastActivityId);
+		const int64 FirstActivityIdToFetch = FMath::Max<int64>(1, LastActivityId - MaximumNumberOfActivities);
+		Database->EnumerateActivitiesInRange(FirstActivityIdToFetch, MaximumNumberOfActivities, [this, &Database, &OutEndpointClientInfoMap, &OutFetchedActivities](FConcertSyncActivity&& InActivity)
+		{
+			if (!OutEndpointClientInfoMap.Contains(InActivity.EndpointId))
+			{
+				FConcertSyncEndpointData EndpointData;
+				if (Database->GetEndpoint(InActivity.EndpointId, EndpointData))
+				{
+					OutEndpointClientInfoMap.Add(InActivity.EndpointId, EndpointData.ClientInfo);
+				}
+			}
+
+			FStructOnScope ActivitySummary;
+			if (InActivity.EventSummary.GetPayload(ActivitySummary))
+			{
+				OutFetchedActivities.Emplace(MoveTemp(InActivity), MoveTemp(ActivitySummary));
+			}
+
+			return EBreakBehavior::Continue;
+		});
+	}
+}
+
+bool FServerSessionHistoryController::GetPackageEvent(const FConcertSessionActivity& Activity, FConcertSyncPackageEventMetaData& OutPackageEvent) const
+{
+	if (const TOptional<FConcertSyncSessionDatabaseNonNullPtr> Database = SyncServer->GetLiveSessionDatabase(InspectedSession->GetId()))
+	{
+		return Database->GetPackageEventMetaData(Activity.Activity.EventId, OutPackageEvent.PackageRevision, OutPackageEvent.PackageInfo);
+	}
+	
+	return false;
+}
+
+TFuture<TOptional<FConcertSyncTransactionEvent>> FServerSessionHistoryController::GetTransactionEvent(const FConcertSessionActivity& Activity) const
+{
+	if (const TOptional<FConcertSyncSessionDatabaseNonNullPtr> Database = SyncServer->GetLiveSessionDatabase(InspectedSession->GetId()))
+	{
+		return FindOrRequestTransactionEvent(*Database, Activity.Activity.EventId);
+	}
+	
+	return MakeFulfilledPromise<TOptional<FConcertSyncTransactionEvent>>().GetFuture(); // Not found.
+}
+
+TFuture<TOptional<FConcertSyncTransactionEvent>> FServerSessionHistoryController::FindOrRequestTransactionEvent(const FConcertSyncSessionDatabase& Database, const int64 TransactionEventId) const
+{
+	FConcertSyncTransactionEvent TransactionEvent;
+	return Database.GetTransactionEvent(TransactionEventId, TransactionEvent, false)
+		?	MakeFulfilledPromise<TOptional<FConcertSyncTransactionEvent>>(MoveTemp(TransactionEvent)).GetFuture()
+		:	MakeFulfilledPromise<TOptional<FConcertSyncTransactionEvent>>().GetFuture();
+}
+
+void FServerSessionHistoryController::OnSessionProduced(const FConcertSyncActivity& ProducedActivity)
+{
+	ReloadActivities();
+}
diff --git a/Engine/Plugins/Developer/Concert/ConcertApp/MultiUserServer/Source/MultiUserServer/Private/Widgets/SessionTabs/Live/ServerSessionHistoryController.h b/Engine/Plugins/Developer/Concert/ConcertApp/MultiUserServer/Source/MultiUserServer/Private/Widgets/SessionTabs/Live/ServerSessionHistoryController.h
new file mode 100644
--- /dev/null
+++ b/Engine/Plugins/Developer/Concert/ConcertApp/MultiUserServer/Source/MultiUserServer/Private/Widgets/SessionTabs/Live/ServerSessionHistoryController.h
@@ -0,0 +1,36 @@
+﻿// Copyright Epic Games, Inc. All Rights Reserved.
+
+#pragma once
+
+#include "CoreMinimal.h"
+#include "Session/History/AbstractSessionHistoryController.h"
+
+struct FConcertSyncActivity;
+class FConcertSyncSessionDatabase;
+class IConcertSyncServer;
+class IConcertServerSession;
+
+class FServerSessionHistoryController : public FAbstractSessionHistoryController
+{
+public:
+
+	FServerSessionHistoryController(TSharedRef<IConcertServerSession> InspectedSession, TSharedRef<IConcertSyncServer> SyncServer);
+	virtual ~FServerSessionHistoryController() override;
+	
+protected:
+
+	//~ Begin FAbstractSessionHistoryController Interface
+	virtual void GetActivities(int64 MaximumNumberOfActivities, TMap<FGuid, FConcertClientInfo>& OutEndpointClientInfoMap, TArray<FConcertSessionActivity>& OutFetchedActivities) const override;
+	virtual bool GetPackageEvent(const FConcertSessionActivity& Activity, FConcertSyncPackageEventMetaData& OutPackageEvent) const override;
+	virtual TFuture<TOptional<FConcertSyncTransactionEvent>> GetTransactionEvent(const FConcertSessionActivity& Activity) const override;
+	//~ End FAbstractSessionHistoryController Interface
+
+private:
+
+	TSharedRef<IConcertServerSession> InspectedSession;
+	TSharedRef<IConcertSyncServer> SyncServer;
+
+	TFuture<TOptional<FConcertSyncTransactionEvent>> FindOrRequestTransactionEvent(const FConcertSyncSessionDatabase& Database, const int64 TransactionEventId) const;
+	
+	void OnSessionProduced(const FConcertSyncActivity& ProducedActivity);
+};
diff --git a/Engine/Plugins/Developer/Concert/ConcertApp/MultiUserServer/Source/MultiUserServer/Private/Widgets/StatusBar/SConcertStatusBar.cpp b/Engine/Plugins/Developer/Concert/ConcertApp/MultiUserServer/Source/MultiUserServer/Private/Widgets/StatusBar/SConcertStatusBar.cpp
new file mode 100644
--- /dev/null
+++ b/Engine/Plugins/Developer/Concert/ConcertApp/MultiUserServer/Source/MultiUserServer/Private/Widgets/StatusBar/SConcertStatusBar.cpp
@@ -0,0 +1,154 @@
+﻿// Copyright Epic Games, Inc. All Rights Reserved.
+
+#include "SConcertStatusBar.h"
+
+#include "OutputLogCreationParams.h"
+
+#include "Framework/Application/SlateApplication.h"
+#include "OutputLogModule.h"
+#include "SWidgetDrawer.h"
+#include "Widgets/Input/SMultiLineEditableTextBox.h"
+#include "Widgets/Layout/SBox.h"
+
+#define LOCTEXT_NAMESPACE "UnrealMultiUserUI"
+
+namespace UE::ConcertServerUI::Private
+{
+	const FName OutputLogId("OutputLog");
+	
+	class FStatusBarSingleton
+	{
+		TSharedPtr<SWidget> StatusBarOutputLog;
+		TArray<TWeakPtr<SWidgetDrawer>> StatusBars;
+		
+		TSharedRef<SWidget> OnGetOutputLog()
+		{
+			if (!StatusBarOutputLog)
+			{
+				FOutputLogCreationParams Params;
+				Params.bCreateDockInLayoutButton = true;
+				Params.SettingsMenuCreationFlags = EOutputLogSettingsMenuFlags::SkipClearOnPie
+					| EOutputLogSettingsMenuFlags::SkipOpenSourceButton
+					| EOutputLogSettingsMenuFlags::SkipEnableWordWrapping; // Checkbox relies on saving an editor config file and does not work correctly
+				Params.AllowAsInitialLogCategory = FAllowLogCategoryCallback::CreateLambda([](const FName CategoryName)
+				{
+					// Hide these categories for improved UX
+					static TArray<FString> DeselectedByDefault = { "Slate", "WindowsTextInput" };
+					const FString CategoryNameAsString = CategoryName.ToString();
+					for (const FString& DeselectedCategory : DeselectedByDefault)
+					{
+						if (CategoryNameAsString.Contains(DeselectedCategory))
+						{
+							return false;
+						}
+					}
+					return true;
+				});
+				StatusBarOutputLog = FOutputLogModule::Get().MakeOutputLogWidget(Params);
+			}
+
+			return StatusBarOutputLog.ToSharedRef();
+		}
+		
+		void OnOutputLogOpened(FName StatusBarWithDrawerName)
+		{
+			// Dismiss all other open drawers - StatusBarOutputLog is shared and shouldn't be in the layout twice
+			for (TWeakPtr<SWidgetDrawer> WidgetDrawer : StatusBars)
+			{
+				if (WidgetDrawer.IsValid())
+				{
+					TSharedPtr<SWidgetDrawer> PinnedDrawer = WidgetDrawer.Pin();
+					if (StatusBarWithDrawerName != PinnedDrawer->GetDrawerName() || PinnedDrawer->IsAnyOtherDrawerOpened(OutputLogId))
+					{
+						PinnedDrawer->CloseDrawerImmediately();
+					}
+				}
+			}
+			
+			FOutputLogModule::Get().FocusOutputLogConsoleBox(StatusBarOutputLog.ToSharedRef());
+		}
+		
+		void OnOutputLogDismised(const TSharedPtr<SWidget>& NewlyFocusedWidget)
+		{}
+
+		void PreShutdownSlate()
+		{
+			StatusBarOutputLog.Reset();
+		}
+
+	public:
+
+		void Init(TSharedRef<SWidgetDrawer> WidgetDrawer, FWidgetDrawerConfig& OutputLogDrawer)
+		{
+			if (!FSlateApplication::Get().OnPreShutdown().IsBoundToObject(this))
+			{
+				// Destroying StatusBarOutputLog in ~FStatusBarSingleton is too late: it causes a crash
+				FSlateApplication::Get().OnPreShutdown().AddRaw(this, &FStatusBarSingleton::PreShutdownSlate);
+			}
+			
+			StatusBars.Add(WidgetDrawer);
+			
+			OutputLogDrawer.GetDrawerContentDelegate.BindRaw(this, &FStatusBarSingleton::OnGetOutputLog);
+			OutputLogDrawer.OnDrawerOpenedDelegate.BindRaw(this, &FStatusBarSingleton::OnOutputLogOpened);
+			OutputLogDrawer.OnDrawerDismissedDelegate.BindRaw(this, &FStatusBarSingleton::OnOutputLogDismised);
+		}
+
+		void Remove(TSharedRef<SWidgetDrawer> WidgetDrawer)
+		{
+			StatusBars.RemoveSingle(WidgetDrawer);
+		}
+		
+	} GStatusBarManager;
+}
+
+SConcertStatusBar::~SConcertStatusBar()
+{
+	UE::ConcertServerUI::Private::GStatusBarManager.Remove(WidgetDrawer.ToSharedRef());
+}
+
+void SConcertStatusBar::Construct(const FArguments& InArgs, FName StatusBarId)
+{
+	ChildSlot
+	[
+		SNew(SBox)
+		.HeightOverride(FAppStyle::Get().GetFloat("StatusBar.Height"))
+		[
+			MakeWidgetDrawer(StatusBarId)
+		]
+	];
+}
+
+TSharedRef<SWidgetDrawer> SConcertStatusBar::MakeWidgetDrawer(FName StatusBarId)
+{
+	using namespace UE::ConcertServerUI::Private;
+	
+	WidgetDrawer = SNew(SWidgetDrawer, StatusBarId);
+
+	TSharedPtr<SMultiLineEditableTextBox> ConsoleEditBox;
+	FSimpleDelegate OnConsoleClosed;
+	FSimpleDelegate OnConsoleCommandExecuted;
+	const TSharedRef<SWidget> OutputLog = 
+			SNew(SBorder)
+			.BorderImage(FAppStyle::Get().GetBrush("Brushes.Panel"))
+			.VAlign(VAlign_Center)
+			.Padding(FMargin(6.0f, 0.0f))
+			[
+				SNew(SBox)
+				.WidthOverride(350.f)
+				[
+					FOutputLogModule::Get().MakeConsoleInputBox(ConsoleEditBox, OnConsoleClosed, OnConsoleCommandExecuted)
+				]
+			];
+	
+	FWidgetDrawerConfig OutputLogDrawer(OutputLogId);
+	GStatusBarManager.Init(WidgetDrawer.ToSharedRef(), OutputLogDrawer);
+	OutputLogDrawer.CustomWidget = OutputLog;
+
+	OutputLogDrawer.ButtonText = LOCTEXT("StatusBar_OutputLogButton", "Output Log");
+	OutputLogDrawer.Icon = FAppStyle::Get().GetBrush("Log.TabIcon");
+	WidgetDrawer->RegisterDrawer(MoveTemp(OutputLogDrawer));
+
+	return WidgetDrawer.ToSharedRef();
+}
+
+#undef LOCTEXT_NAMESPACE
diff --git a/Engine/Plugins/Developer/Concert/ConcertApp/MultiUserServer/Source/MultiUserServer/Private/Widgets/StatusBar/SConcertStatusBar.h b/Engine/Plugins/Developer/Concert/ConcertApp/MultiUserServer/Source/MultiUserServer/Private/Widgets/StatusBar/SConcertStatusBar.h
new file mode 100644
--- /dev/null
+++ b/Engine/Plugins/Developer/Concert/ConcertApp/MultiUserServer/Source/MultiUserServer/Private/Widgets/StatusBar/SConcertStatusBar.h
@@ -0,0 +1,26 @@
+﻿// Copyright Epic Games, Inc. All Rights Reserved.
+
+#pragma once
+
+#include "CoreMinimal.h"
+#include "SWidgetDrawer.h"
+#include "Widgets/DeclarativeSyntaxSupport.h"
+#include "Widgets/SCompoundWidget.h"
+
+class SConcertStatusBar : public SCompoundWidget
+{
+public:
+
+	SLATE_BEGIN_ARGS(SConcertStatusBar) { }
+	SLATE_END_ARGS()
+	
+	~SConcertStatusBar();
+
+	void Construct(const FArguments& InArgs, FName StatusBarId);
+
+private:
+
+	TSharedPtr<SWidgetDrawer> WidgetDrawer;
+
+	TSharedRef<SWidgetDrawer> MakeWidgetDrawer(FName StatusBarId);
+};
diff --git a/Engine/Plugins/Developer/Concert/ConcertApp/MultiUserServer/Source/MultiUserServer/Public/IMultiUserServerModule.h b/Engine/Plugins/Developer/Concert/ConcertApp/MultiUserServer/Source/MultiUserServer/Public/IMultiUserServerModule.h
new file mode 100644
--- /dev/null
+++ b/Engine/Plugins/Developer/Concert/ConcertApp/MultiUserServer/Source/MultiUserServer/Public/IMultiUserServerModule.h
@@ -0,0 +1,41 @@
+// Copyright Epic Games, Inc. All Rights Reserved.
+
+#pragma once
+
+#include "CoreMinimal.h"
+#include "Modules/ModuleInterface.h"
+#include "Modules/ModuleManager.h"
+
+class IConcertSyncServer;
+struct FConcertSyncServerLoopInitArgs;
+
+class IMultiUserServerModule : public IModuleInterface
+{
+public:
+	
+	/**
+	 * Singleton-like access to this module's interface.  This is just for convenience!
+	 * Beware of calling this during the shutdown phase, though.  Your module might have been unloaded already.
+	 *
+	 * @return Returns singleton instance, loading the module on demand if needed
+	 */
+	static inline IMultiUserServerModule& Get()
+	{
+		static const FName ModuleName = "MultiUserServer";
+		return FModuleManager::LoadModuleChecked<IMultiUserServerModule>(ModuleName);
+	}
+
+	/**
+	 * Checks to see if this module is loaded and ready.  It is only valid to call Get() during shutdown if IsAvailable() returns true.
+	 *
+	 * @return True if the module is loaded and ready to use
+	 */
+	static inline bool IsAvailable()
+	{
+		static const FName ModuleName = "MultiUserServer";
+		return FModuleManager::Get().IsModuleLoaded(ModuleName);
+	}
+
+	/** Set-up callbacks for creating Slate UI when the MU server is started. Called before the server loop is initialized. */
+	virtual void InitSlateForServer(FConcertSyncServerLoopInitArgs& InitArgs) = 0;
+};
diff --git a/Engine/Plugins/Developer/Concert/ConcertUI/ConcertServerUI/ConcertServerUI.uplugin b/Engine/Plugins/Developer/Concert/ConcertUI/ConcertServerUI/ConcertServerUI.uplugin
deleted file mode 100644
diff --git a/Engine/Plugins/Developer/Concert/ConcertUI/ConcertServerUI/Source/ConcertServerUI/ConcertServerUI.Build.cs b/Engine/Plugins/Developer/Concert/ConcertUI/ConcertServerUI/Source/ConcertServerUI/ConcertServerUI.Build.cs
deleted file mode 100644
diff --git a/Engine/Plugins/Developer/Concert/ConcertUI/ConcertServerUI/Source/ConcertServerUI/Private/ConcertConsoleCommandExecutor.cpp b/Engine/Plugins/Developer/Concert/ConcertUI/ConcertServerUI/Source/ConcertServerUI/Private/ConcertConsoleCommandExecutor.cpp
deleted file mode 100644
diff --git a/Engine/Plugins/Developer/Concert/ConcertUI/ConcertServerUI/Source/ConcertServerUI/Private/ConcertConsoleCommandExecutor.h b/Engine/Plugins/Developer/Concert/ConcertUI/ConcertServerUI/Source/ConcertServerUI/Private/ConcertConsoleCommandExecutor.h
deleted file mode 100644
diff --git a/Engine/Plugins/Developer/Concert/ConcertUI/ConcertServerUI/Source/ConcertServerUI/Private/ConcertServerStyle.cpp b/Engine/Plugins/Developer/Concert/ConcertUI/ConcertServerUI/Source/ConcertServerUI/Private/ConcertServerStyle.cpp
deleted file mode 100644
diff --git a/Engine/Plugins/Developer/Concert/ConcertUI/ConcertServerUI/Source/ConcertServerUI/Private/ConcertServerStyle.h b/Engine/Plugins/Developer/Concert/ConcertUI/ConcertServerUI/Source/ConcertServerUI/Private/ConcertServerStyle.h
deleted file mode 100644
diff --git a/Engine/Plugins/Developer/Concert/ConcertUI/ConcertServerUI/Source/ConcertServerUI/Private/ConcertServerUIModule.cpp b/Engine/Plugins/Developer/Concert/ConcertUI/ConcertServerUI/Source/ConcertServerUI/Private/ConcertServerUIModule.cpp
deleted file mode 100644
diff --git a/Engine/Plugins/Developer/Concert/ConcertUI/ConcertServerUI/Source/ConcertServerUI/Private/Widgets/Browser/ConcertServerSessionBrowserController.cpp b/Engine/Plugins/Developer/Concert/ConcertUI/ConcertServerUI/Source/ConcertServerUI/Private/Widgets/Browser/ConcertServerSessionBrowserController.cpp
deleted file mode 100644
diff --git a/Engine/Plugins/Developer/Concert/ConcertUI/ConcertServerUI/Source/ConcertServerUI/Private/Widgets/Browser/ConcertServerSessionBrowserController.h b/Engine/Plugins/Developer/Concert/ConcertUI/ConcertServerUI/Source/ConcertServerUI/Private/Widgets/Browser/ConcertServerSessionBrowserController.h
deleted file mode 100644
diff --git a/Engine/Plugins/Developer/Concert/ConcertUI/ConcertServerUI/Source/ConcertServerUI/Private/Widgets/Browser/SConcertServerSessionBrowser.cpp b/Engine/Plugins/Developer/Concert/ConcertUI/ConcertServerUI/Source/ConcertServerUI/Private/Widgets/Browser/SConcertServerSessionBrowser.cpp
deleted file mode 100644
diff --git a/Engine/Plugins/Developer/Concert/ConcertUI/ConcertServerUI/Source/ConcertServerUI/Private/Widgets/Browser/SConcertServerSessionBrowser.h b/Engine/Plugins/Developer/Concert/ConcertUI/ConcertServerUI/Source/ConcertServerUI/Private/Widgets/Browser/SConcertServerSessionBrowser.h
deleted file mode 100644
diff --git a/Engine/Plugins/Developer/Concert/ConcertUI/ConcertServerUI/Source/ConcertServerUI/Private/Widgets/ConcertServerTabs.cpp b/Engine/Plugins/Developer/Concert/ConcertUI/ConcertServerUI/Source/ConcertServerUI/Private/Widgets/ConcertServerTabs.cpp
deleted file mode 100644
diff --git a/Engine/Plugins/Developer/Concert/ConcertUI/ConcertServerUI/Source/ConcertServerUI/Private/Widgets/ConcertServerTabs.h b/Engine/Plugins/Developer/Concert/ConcertUI/ConcertServerUI/Source/ConcertServerUI/Private/Widgets/ConcertServerTabs.h
deleted file mode 100644
diff --git a/Engine/Plugins/Developer/Concert/ConcertUI/ConcertServerUI/Source/ConcertServerUI/Private/Widgets/ConcertServerWindowController.cpp b/Engine/Plugins/Developer/Concert/ConcertUI/ConcertServerUI/Source/ConcertServerUI/Private/Widgets/ConcertServerWindowController.cpp
deleted file mode 100644
diff --git a/Engine/Plugins/Developer/Concert/ConcertUI/ConcertServerUI/Source/ConcertServerUI/Private/Widgets/ConcertServerWindowController.h b/Engine/Plugins/Developer/Concert/ConcertUI/ConcertServerUI/Source/ConcertServerUI/Private/Widgets/ConcertServerWindowController.h
deleted file mode 100644
diff --git a/Engine/Plugins/Developer/Concert/ConcertUI/ConcertServerUI/Source/ConcertServerUI/Private/Widgets/IConcertComponent.h b/Engine/Plugins/Developer/Concert/ConcertUI/ConcertServerUI/Source/ConcertServerUI/Private/Widgets/IConcertComponent.h
deleted file mode 100644
diff --git a/Engine/Plugins/Developer/Concert/ConcertUI/ConcertServerUI/Source/ConcertServerUI/Private/Widgets/SessionTabs/AbstractConcertSessionTab.cpp b/Engine/Plugins/Developer/Concert/ConcertUI/ConcertServerUI/Source/ConcertServerUI/Private/Widgets/SessionTabs/AbstractConcertSessionTab.cpp
deleted file mode 100644
diff --git a/Engine/Plugins/Developer/Concert/ConcertUI/ConcertServerUI/Source/ConcertServerUI/Private/Widgets/SessionTabs/AbstractConcertSessionTab.h b/Engine/Plugins/Developer/Concert/ConcertUI/ConcertServerUI/Source/ConcertServerUI/Private/Widgets/SessionTabs/AbstractConcertSessionTab.h
deleted file mode 100644
diff --git a/Engine/Plugins/Developer/Concert/ConcertUI/ConcertServerUI/Source/ConcertServerUI/Private/Widgets/SessionTabs/Archived/ArchivedConcertSessionTab.cpp b/Engine/Plugins/Developer/Concert/ConcertUI/ConcertServerUI/Source/ConcertServerUI/Private/Widgets/SessionTabs/Archived/ArchivedConcertSessionTab.cpp
deleted file mode 100644
diff --git a/Engine/Plugins/Developer/Concert/ConcertUI/ConcertServerUI/Source/ConcertServerUI/Private/Widgets/SessionTabs/Archived/ArchivedConcertSessionTab.h b/Engine/Plugins/Developer/Concert/ConcertUI/ConcertServerUI/Source/ConcertServerUI/Private/Widgets/SessionTabs/Archived/ArchivedConcertSessionTab.h
deleted file mode 100644
diff --git a/Engine/Plugins/Developer/Concert/ConcertUI/ConcertServerUI/Source/ConcertServerUI/Private/Widgets/SessionTabs/Archived/SConcertArchivedSessionInspector.cpp b/Engine/Plugins/Developer/Concert/ConcertUI/ConcertServerUI/Source/ConcertServerUI/Private/Widgets/SessionTabs/Archived/SConcertArchivedSessionInspector.cpp
deleted file mode 100644
diff --git a/Engine/Plugins/Developer/Concert/ConcertUI/ConcertServerUI/Source/ConcertServerUI/Private/Widgets/SessionTabs/Archived/SConcertArchivedSessionInspector.h b/Engine/Plugins/Developer/Concert/ConcertUI/ConcertServerUI/Source/ConcertServerUI/Private/Widgets/SessionTabs/Archived/SConcertArchivedSessionInspector.h
deleted file mode 100644
diff --git a/Engine/Plugins/Developer/Concert/ConcertUI/ConcertServerUI/Source/ConcertServerUI/Private/Widgets/SessionTabs/Live/LiveConcertSessionTab.cpp b/Engine/Plugins/Developer/Concert/ConcertUI/ConcertServerUI/Source/ConcertServerUI/Private/Widgets/SessionTabs/Live/LiveConcertSessionTab.cpp
deleted file mode 100644
diff --git a/Engine/Plugins/Developer/Concert/ConcertUI/ConcertServerUI/Source/ConcertServerUI/Private/Widgets/SessionTabs/Live/LiveConcertSessionTab.h b/Engine/Plugins/Developer/Concert/ConcertUI/ConcertServerUI/Source/ConcertServerUI/Private/Widgets/SessionTabs/Live/LiveConcertSessionTab.h
deleted file mode 100644
diff --git a/Engine/Plugins/Developer/Concert/ConcertUI/ConcertServerUI/Source/ConcertServerUI/Private/Widgets/SessionTabs/Live/PackageViewer/ConcertSessionPackageViewerController.cpp b/Engine/Plugins/Developer/Concert/ConcertUI/ConcertServerUI/Source/ConcertServerUI/Private/Widgets/SessionTabs/Live/PackageViewer/ConcertSessionPackageViewerController.cpp
deleted file mode 100644
diff --git a/Engine/Plugins/Developer/Concert/ConcertUI/ConcertServerUI/Source/ConcertServerUI/Private/Widgets/SessionTabs/Live/PackageViewer/ConcertSessionPackageViewerController.h b/Engine/Plugins/Developer/Concert/ConcertUI/ConcertServerUI/Source/ConcertServerUI/Private/Widgets/SessionTabs/Live/PackageViewer/ConcertSessionPackageViewerController.h
deleted file mode 100644
diff --git a/Engine/Plugins/Developer/Concert/ConcertUI/ConcertServerUI/Source/ConcertServerUI/Private/Widgets/SessionTabs/Live/PackageViewer/SConcertSessionPackageViewer.cpp b/Engine/Plugins/Developer/Concert/ConcertUI/ConcertServerUI/Source/ConcertServerUI/Private/Widgets/SessionTabs/Live/PackageViewer/SConcertSessionPackageViewer.cpp
deleted file mode 100644
diff --git a/Engine/Plugins/Developer/Concert/ConcertUI/ConcertServerUI/Source/ConcertServerUI/Private/Widgets/SessionTabs/Live/PackageViewer/SConcertSessionPackageViewer.h b/Engine/Plugins/Developer/Concert/ConcertUI/ConcertServerUI/Source/ConcertServerUI/Private/Widgets/SessionTabs/Live/PackageViewer/SConcertSessionPackageViewer.h
deleted file mode 100644
diff --git a/Engine/Plugins/Developer/Concert/ConcertUI/ConcertServerUI/Source/ConcertServerUI/Private/Widgets/SessionTabs/Live/SConcertSessionInspector.cpp b/Engine/Plugins/Developer/Concert/ConcertUI/ConcertServerUI/Source/ConcertServerUI/Private/Widgets/SessionTabs/Live/SConcertSessionInspector.cpp
deleted file mode 100644
diff --git a/Engine/Plugins/Developer/Concert/ConcertUI/ConcertServerUI/Source/ConcertServerUI/Private/Widgets/SessionTabs/Live/SConcertSessionInspector.h b/Engine/Plugins/Developer/Concert/ConcertUI/ConcertServerUI/Source/ConcertServerUI/Private/Widgets/SessionTabs/Live/SConcertSessionInspector.h
deleted file mode 100644
diff --git a/Engine/Plugins/Developer/Concert/ConcertUI/ConcertServerUI/Source/ConcertServerUI/Private/Widgets/SessionTabs/Live/ServerSessionHistoryController.cpp b/Engine/Plugins/Developer/Concert/ConcertUI/ConcertServerUI/Source/ConcertServerUI/Private/Widgets/SessionTabs/Live/ServerSessionHistoryController.cpp
deleted file mode 100644
diff --git a/Engine/Plugins/Developer/Concert/ConcertUI/ConcertServerUI/Source/ConcertServerUI/Private/Widgets/SessionTabs/Live/ServerSessionHistoryController.h b/Engine/Plugins/Developer/Concert/ConcertUI/ConcertServerUI/Source/ConcertServerUI/Private/Widgets/SessionTabs/Live/ServerSessionHistoryController.h
deleted file mode 100644
diff --git a/Engine/Plugins/Developer/Concert/ConcertUI/ConcertServerUI/Source/ConcertServerUI/Private/Widgets/StatusBar/SConcertStatusBar.cpp b/Engine/Plugins/Developer/Concert/ConcertUI/ConcertServerUI/Source/ConcertServerUI/Private/Widgets/StatusBar/SConcertStatusBar.cpp
deleted file mode 100644
diff --git a/Engine/Plugins/Developer/Concert/ConcertUI/ConcertServerUI/Source/ConcertServerUI/Private/Widgets/StatusBar/SConcertStatusBar.h b/Engine/Plugins/Developer/Concert/ConcertUI/ConcertServerUI/Source/ConcertServerUI/Private/Widgets/StatusBar/SConcertStatusBar.h
deleted file mode 100644
diff --git a/Engine/Plugins/Developer/Concert/ConcertUI/ConcertServerUI/Source/ConcertServerUI/Public/IConcertServerUIModule.h b/Engine/Plugins/Developer/Concert/ConcertUI/ConcertServerUI/Source/ConcertServerUI/Public/IConcertServerUIModule.h
deleted file mode 100644
diff --git a/Engine/Programs/UnrealMultiUserServer/Config/DefaultEngine.ini b/Engine/Programs/UnrealMultiUserServer/Config/DefaultEngine.ini
--- a/Engine/Programs/UnrealMultiUserServer/Config/DefaultEngine.ini
+++ b/Engine/Programs/UnrealMultiUserServer/Config/DefaultEngine.ini
@@ -4,4 +4,4 @@ GameNameShort=UCS
 [Plugins]
 +ProgramEnabledPlugins="UdpMessaging"
 +ProgramEnabledPlugins="ConcertSyncServer"
-+ProgramEnabledPlugins="ConcertServerUI"
++ProgramEnabledPlugins="MultiUserServer"
diff --git a/Engine/Source/Programs/UnrealMultiUserServer/Private/UnrealMultiUserServerRun.cpp b/Engine/Source/Programs/UnrealMultiUserServer/Private/UnrealMultiUserServerRun.cpp
--- a/Engine/Source/Programs/UnrealMultiUserServer/Private/UnrealMultiUserServerRun.cpp
+++ b/Engine/Source/Programs/UnrealMultiUserServer/Private/UnrealMultiUserServerRun.cpp
@@ -4,7 +4,7 @@
 
 #include "ConcertSettings.h"
 #include "ConcertSyncServerLoop.h"
-#include "IConcertServerUIModule.h"
+#include "IMultiUserServerModule.h"
 
 #include "Misc/Parse.h"
 
@@ -23,14 +23,14 @@ namespace UE::UnrealMultiUserServer
 			ServerLoopInitArgs.bShowConsole = false;
 			ServerLoopInitArgs.PreInitServerLoop.AddLambda([&ServerLoopInitArgs]()	
 			{
-				TSharedPtr<IPlugin> Plugin = IPluginManager::Get().FindPlugin(TEXT("ConcertServerUI"));
+				TSharedPtr<IPlugin> Plugin = IPluginManager::Get().FindPlugin(TEXT("MultiUserServer"));
 				if (!Plugin || !Plugin->IsEnabled())
 				{
 					UE_LOG(LogSyncServer, Error, TEXT("The 'ConcertServerUI' plugin is disabled."));
 				}
 				else
 				{
-					IConcertServerUIModule::Get().InitSlateForServer(ServerLoopInitArgs);
+					IMultiUserServerModule::Get().InitSlateForServer(ServerLoopInitArgs);
 				}
 			});
 		}
diff --git a/Engine/Source/Programs/UnrealMultiUserServer/UnrealMultiUserServer.Build.cs b/Engine/Source/Programs/UnrealMultiUserServer/UnrealMultiUserServer.Build.cs
--- a/Engine/Source/Programs/UnrealMultiUserServer/UnrealMultiUserServer.Build.cs
+++ b/Engine/Source/Programs/UnrealMultiUserServer/UnrealMultiUserServer.Build.cs
@@ -22,7 +22,7 @@ public class UnrealMultiUserServer : ModuleRules
 				"Launch",
 				"ConcertSyncCore",
 				"ConcertSyncServer",
-				"ConcertServerUI",
+				"MultiUserServer",
 			}
 		);
 
@@ -35,7 +35,7 @@ public class UnrealMultiUserServer : ModuleRules
 		DynamicallyLoadedModuleNames.AddRange(
 			new string[] {
 				"ConcertSyncServer",
-				"ConcertServerUI"
+				"MultiUserServer"
 			}
 		);
 	}
