Rename AbstractConcertSessionTab

diff --git a/Engine/Plugins/Developer/Concert/ConcertApp/MultiUserServer/Source/MultiUserServer/Private/Widgets/ConcertServerWindowController.cpp b/Engine/Plugins/Developer/Concert/ConcertApp/MultiUserServer/Source/MultiUserServer/Private/Widgets/ConcertServerWindowController.cpp
--- a/Engine/Plugins/Developer/Concert/ConcertApp/MultiUserServer/Source/MultiUserServer/Private/Widgets/ConcertServerWindowController.cpp
+++ b/Engine/Plugins/Developer/Concert/ConcertApp/MultiUserServer/Source/MultiUserServer/Private/Widgets/ConcertServerWindowController.cpp
@@ -82,15 +82,15 @@ void FConcertServerWindowController::CreateWindow()
 
 void FConcertServerWindowController::OpenSessionTab(const FGuid& SessionId)
 {
-	if (TSharedPtr<FAbstractConcertSessionTab> SessionTab = GetOrRegisterSessionTab(SessionId))
+	if (TSharedPtr<FConcertSessionTabBase> SessionTab = GetOrRegisterSessionTab(SessionId))
 	{
 		SessionTab->OpenSessionTab();
 	}
 }
 
-TSharedPtr<FAbstractConcertSessionTab> FConcertServerWindowController::GetOrRegisterSessionTab(const FGuid& SessionId)
+TSharedPtr<FConcertSessionTabBase> FConcertServerWindowController::GetOrRegisterSessionTab(const FGuid& SessionId)
 {
-	if (TSharedRef<FAbstractConcertSessionTab>* FoundId = RegisteredSessions.Find(SessionId))
+	if (TSharedRef<FConcertSessionTabBase>* FoundId = RegisteredSessions.Find(SessionId))
 	{
 		return *FoundId;
 	}
diff --git a/Engine/Plugins/Developer/Concert/ConcertApp/MultiUserServer/Source/MultiUserServer/Private/Widgets/ConcertServerWindowController.h b/Engine/Plugins/Developer/Concert/ConcertApp/MultiUserServer/Source/MultiUserServer/Private/Widgets/ConcertServerWindowController.h
--- a/Engine/Plugins/Developer/Concert/ConcertApp/MultiUserServer/Source/MultiUserServer/Private/Widgets/ConcertServerWindowController.h
+++ b/Engine/Plugins/Developer/Concert/ConcertApp/MultiUserServer/Source/MultiUserServer/Private/Widgets/ConcertServerWindowController.h
@@ -5,7 +5,7 @@
 #include "CoreMinimal.h"
 #include "Framework/Docking/TabManager.h"
 
-class FAbstractConcertSessionTab;
+class FConcertSessionTabBase;
 class FLiveConcertSessionTab;
 class IConcertServerSession;
 class IConcertSyncServer;
@@ -49,7 +49,7 @@ private:
 
 	/** The main window being managed */
 	TSharedPtr<SWindow> RootWindow;
-	TMap<FGuid, TSharedRef<FAbstractConcertSessionTab>> RegisteredSessions;
+	TMap<FGuid, TSharedRef<FConcertSessionTabBase>> RegisteredSessions;
 	
 	/** Manages the session browser */
 	TSharedPtr<FConcertServerSessionBrowserController> SessionBrowserController;
@@ -57,7 +57,7 @@ private:
 	void InitComponents();
 
 	/** Gets the manager for a session tab if the session ID is valid */
-	TSharedPtr<FAbstractConcertSessionTab> GetOrRegisterSessionTab(const FGuid& SessionId);
+	TSharedPtr<FConcertSessionTabBase> GetOrRegisterSessionTab(const FGuid& SessionId);
 
 	void OnWindowClosed(const TSharedRef<SWindow>& Window);
 	void SaveLayout() const;
diff --git a/Engine/Plugins/Developer/Concert/ConcertApp/MultiUserServer/Source/MultiUserServer/Private/Widgets/SessionTabs/AbstractConcertSessionTab.cpp b/Engine/Plugins/Developer/Concert/ConcertApp/MultiUserServer/Source/MultiUserServer/Private/Widgets/SessionTabs/AbstractConcertSessionTab.cpp
deleted file mode 100644
diff --git a/Engine/Plugins/Developer/Concert/ConcertApp/MultiUserServer/Source/MultiUserServer/Private/Widgets/SessionTabs/AbstractConcertSessionTab.h b/Engine/Plugins/Developer/Concert/ConcertApp/MultiUserServer/Source/MultiUserServer/Private/Widgets/SessionTabs/AbstractConcertSessionTab.h
deleted file mode 100644
diff --git a/Engine/Plugins/Developer/Concert/ConcertApp/MultiUserServer/Source/MultiUserServer/Private/Widgets/SessionTabs/Archived/ArchivedConcertSessionTab.cpp b/Engine/Plugins/Developer/Concert/ConcertApp/MultiUserServer/Source/MultiUserServer/Private/Widgets/SessionTabs/Archived/ArchivedConcertSessionTab.cpp
--- a/Engine/Plugins/Developer/Concert/ConcertApp/MultiUserServer/Source/MultiUserServer/Private/Widgets/SessionTabs/Archived/ArchivedConcertSessionTab.cpp
+++ b/Engine/Plugins/Developer/Concert/ConcertApp/MultiUserServer/Source/MultiUserServer/Private/Widgets/SessionTabs/Archived/ArchivedConcertSessionTab.cpp
@@ -6,7 +6,7 @@
 #include "Widgets/StatusBar/SConcertStatusBar.h"
 
 FArchivedConcertSessionTab::FArchivedConcertSessionTab(const FGuid& InspectedSessionID, TSharedRef<IConcertSyncServer> SyncServer, TAttribute<TSharedRef<SWindow>> ConstructUnderWindow)
-	: FAbstractConcertSessionTab(SyncServer)
+	: FConcertSessionTabBase(SyncServer)
 	, InspectedSessionID(InspectedSessionID)
 	, ConstructUnderWindow(ConstructUnderWindow)
 {}
diff --git a/Engine/Plugins/Developer/Concert/ConcertApp/MultiUserServer/Source/MultiUserServer/Private/Widgets/SessionTabs/Archived/ArchivedConcertSessionTab.h b/Engine/Plugins/Developer/Concert/ConcertApp/MultiUserServer/Source/MultiUserServer/Private/Widgets/SessionTabs/Archived/ArchivedConcertSessionTab.h
--- a/Engine/Plugins/Developer/Concert/ConcertApp/MultiUserServer/Source/MultiUserServer/Private/Widgets/SessionTabs/Archived/ArchivedConcertSessionTab.h
+++ b/Engine/Plugins/Developer/Concert/ConcertApp/MultiUserServer/Source/MultiUserServer/Private/Widgets/SessionTabs/Archived/ArchivedConcertSessionTab.h
@@ -3,12 +3,12 @@
 #pragma once
 
 #include "CoreMinimal.h"
-#include "SConcertArchivedSessionInspector.h"
+#include "Widgets/SessionTabs/ConcertSessionTabBase.h"
 
-#include "Widgets/SessionTabs/AbstractConcertSessionTab.h"
+class SConcertArchivedSessionInspector;
 
 /** Manages the tab for an archived session.  */
-class FArchivedConcertSessionTab : public FAbstractConcertSessionTab
+class FArchivedConcertSessionTab : public FConcertSessionTabBase
 {
 public:
 
@@ -29,7 +29,7 @@ private:
 
 	/** Used later to obtain the window into which to add the tab */
 	TAttribute<TSharedRef<SWindow>> ConstructUnderWindow;
-
+    
 	/** Displays session */
 	TSharedPtr<SConcertArchivedSessionInspector> Inspector;
 };
diff --git a/Engine/Plugins/Developer/Concert/ConcertApp/MultiUserServer/Source/MultiUserServer/Private/Widgets/SessionTabs/ConcertSessionTabBase.cpp b/Engine/Plugins/Developer/Concert/ConcertApp/MultiUserServer/Source/MultiUserServer/Private/Widgets/SessionTabs/ConcertSessionTabBase.cpp
new file mode 100644
--- /dev/null
+++ b/Engine/Plugins/Developer/Concert/ConcertApp/MultiUserServer/Source/MultiUserServer/Private/Widgets/SessionTabs/ConcertSessionTabBase.cpp
@@ -0,0 +1,65 @@
+﻿// Copyright Epic Games, Inc. All Rights Reserved.
+
+#include "ConcertSessionTabBase.h"
+
+#include "IConcertServer.h"
+#include "IConcertSyncServer.h"
+#include "Widgets/ConcertServerTabs.h"
+
+#include "Framework/Docking/TabManager.h"
+#include "Widgets/Docking/SDockTab.h"
+
+namespace UE::ConcertServerUI::Private
+{
+	static TOptional<FString> GetSessionName(const TSharedRef<IConcertSyncServer>& SyncServer, const FGuid SessionId)
+	{
+		if (const TSharedPtr<IConcertServerSession> LiveSessionInfo = SyncServer->GetConcertServer()->GetLiveSession(SessionId))
+		{
+			return LiveSessionInfo->GetSessionInfo().SessionName;
+		}
+		if (const TOptional<FConcertSessionInfo> ArchivedSessionInfo = SyncServer->GetConcertServer()->GetArchivedSessionInfo(SessionId))
+		{
+			return ArchivedSessionInfo->SessionName;
+		}
+		return {};
+	}
+}
+
+FConcertSessionTabBase::FConcertSessionTabBase(TSharedRef<IConcertSyncServer> SyncServer)
+	: SyncServer(MoveTemp(SyncServer))
+{}
+
+void FConcertSessionTabBase::OpenSessionTab()
+{
+	const TSharedRef<FGlobalTabmanager>& TabManager = FGlobalTabmanager::Get();
+	const FTabId TabId { *GetTabId() };
+
+	EnsureInitDockTab();
+	if (TabManager->FindExistingLiveTab(TabId))
+	{
+		TabManager->DrawAttention(DockTab.ToSharedRef());
+	}
+	else
+	{
+		const FTabManager::FLastMajorOrNomadTab Search(ConcertServerTabs::GetSessionBrowserTabId());
+		FGlobalTabmanager::Get()->InsertNewDocumentTab(*GetTabId(), Search, DockTab.ToSharedRef());
+
+		OnOpenTab();
+	}
+}
+
+void FConcertSessionTabBase::EnsureInitDockTab()
+{
+	if (!DockTab.IsValid())
+	{
+		const TOptional<FString> SessionName = UE::ConcertServerUI::Private::GetSessionName(SyncServer, GetSessionID());
+		check(SessionName);
+		
+		const FText Title = FText::FromString(*SessionName);
+		DockTab = SNew(SDockTab)
+			.Label(Title)
+			.TabRole(MajorTab);
+		
+		CreateDockContent(DockTab.ToSharedRef());
+	}
+}
diff --git a/Engine/Plugins/Developer/Concert/ConcertApp/MultiUserServer/Source/MultiUserServer/Private/Widgets/SessionTabs/ConcertSessionTabBase.h b/Engine/Plugins/Developer/Concert/ConcertApp/MultiUserServer/Source/MultiUserServer/Private/Widgets/SessionTabs/ConcertSessionTabBase.h
new file mode 100644
--- /dev/null
+++ b/Engine/Plugins/Developer/Concert/ConcertApp/MultiUserServer/Source/MultiUserServer/Private/Widgets/SessionTabs/ConcertSessionTabBase.h
@@ -0,0 +1,42 @@
+﻿// Copyright Epic Games, Inc. All Rights Reserved.
+
+#pragma once
+
+#include "CoreMinimal.h"
+
+class IConcertSyncServer;
+class SDockTab;
+class SWindow;
+
+/** Shared functionality for a tab that manages a session */
+class FConcertSessionTabBase
+{
+public:
+
+    FConcertSessionTabBase(TSharedRef<IConcertSyncServer> SyncServer);
+	virtual ~FConcertSessionTabBase() = default;
+	
+	/** Opens or draws attention to the given tab */
+	void OpenSessionTab();
+	
+protected:
+
+	virtual FGuid GetSessionID() const = 0;
+	virtual void CreateDockContent(const TSharedRef<SDockTab>& DockTab) = 0;
+	virtual void OnOpenTab() = 0;
+	
+	/** Generates a tab ID for FTabManager::InsertNewDocumentTab */
+	FString GetTabId() const { return GetSessionID().ToString(); }
+	const TSharedRef<IConcertSyncServer>& GetSyncServer() const { return SyncServer; }
+	
+private:
+
+	/** Used to look up session name */
+	const TSharedRef<IConcertSyncServer> SyncServer;
+	
+	/** The tab containing the UI for InspectedSession */
+	TSharedPtr<SDockTab> DockTab;
+
+	/** Inits DockTab if it has not yet been inited */
+	void EnsureInitDockTab();
+};
diff --git a/Engine/Plugins/Developer/Concert/ConcertApp/MultiUserServer/Source/MultiUserServer/Private/Widgets/SessionTabs/Live/LiveConcertSessionTab.cpp b/Engine/Plugins/Developer/Concert/ConcertApp/MultiUserServer/Source/MultiUserServer/Private/Widgets/SessionTabs/Live/LiveConcertSessionTab.cpp
--- a/Engine/Plugins/Developer/Concert/ConcertApp/MultiUserServer/Source/MultiUserServer/Private/Widgets/SessionTabs/Live/LiveConcertSessionTab.cpp
+++ b/Engine/Plugins/Developer/Concert/ConcertApp/MultiUserServer/Source/MultiUserServer/Private/Widgets/SessionTabs/Live/LiveConcertSessionTab.cpp
@@ -10,7 +10,7 @@
 #include "Widgets/Docking/SDockTab.h"
 
 FLiveConcertSessionTab::FLiveConcertSessionTab(TSharedRef<IConcertServerSession> InspectedSession, TSharedRef<IConcertSyncServer> SyncServer, TAttribute<TSharedRef<SWindow>> ConstructUnderWindow)
-	: FAbstractConcertSessionTab(SyncServer)
+	: FConcertSessionTabBase(SyncServer)
 	, InspectedSession(MoveTemp(InspectedSession))
 	, ConstructUnderWindow(MoveTemp(ConstructUnderWindow))
 	, SessionHistoryController(MakeShared<FServerSessionHistoryController>(InspectedSession, SyncServer))
diff --git a/Engine/Plugins/Developer/Concert/ConcertApp/MultiUserServer/Source/MultiUserServer/Private/Widgets/SessionTabs/Live/LiveConcertSessionTab.h b/Engine/Plugins/Developer/Concert/ConcertApp/MultiUserServer/Source/MultiUserServer/Private/Widgets/SessionTabs/Live/LiveConcertSessionTab.h
--- a/Engine/Plugins/Developer/Concert/ConcertApp/MultiUserServer/Source/MultiUserServer/Private/Widgets/SessionTabs/Live/LiveConcertSessionTab.h
+++ b/Engine/Plugins/Developer/Concert/ConcertApp/MultiUserServer/Source/MultiUserServer/Private/Widgets/SessionTabs/Live/LiveConcertSessionTab.h
@@ -3,9 +3,8 @@
 #pragma once
 
 #include "CoreMinimal.h"
-#include "ServerSessionHistoryController.h"
-#include "Widgets/SessionTabs/AbstractConcertSessionTab.h"
-#include "Types/SlateAttribute.h"
+#include "Widgets/SessionTabs/Live/ServerSessionHistoryController.h"
+#include "Widgets/SessionTabs/ConcertSessionTabBase.h"
 
 class FConcertSessionPackageViewerController;
 class IConcertSyncServer;
@@ -18,7 +17,7 @@ class SWindow;
  * Manages the tab that contains the UI for a session.
  * It has access to the controllers and views needed for displaying a session.
  */
-class FLiveConcertSessionTab : public FAbstractConcertSessionTab
+class FLiveConcertSessionTab : public FConcertSessionTabBase
 {
 public:
 
