Extract session history controller functionality to base class

diff --git a/Engine/Plugins/Developer/Concert/ConcertApp/MultiUserServer/Source/MultiUserServer/Private/Widgets/SessionTabs/Live/LiveConcertSessionTab.cpp b/Engine/Plugins/Developer/Concert/ConcertApp/MultiUserServer/Source/MultiUserServer/Private/Widgets/SessionTabs/Live/LiveConcertSessionTab.cpp
--- a/Engine/Plugins/Developer/Concert/ConcertApp/MultiUserServer/Source/MultiUserServer/Private/Widgets/SessionTabs/Live/LiveConcertSessionTab.cpp
+++ b/Engine/Plugins/Developer/Concert/ConcertApp/MultiUserServer/Source/MultiUserServer/Private/Widgets/SessionTabs/Live/LiveConcertSessionTab.cpp
@@ -3,6 +3,7 @@
 #include "LiveConcertSessionTab.h"
 
 #include "IConcertSession.h"
+#include "LiveServerSessionHistoryController.h"
 #include "PackageViewer/ConcertSessionPackageViewerController.h"
 #include "SConcertSessionInspector.h"
 #include "Widgets/StatusBar/SConcertStatusBar.h"
@@ -13,7 +14,7 @@ FLiveConcertSessionTab::FLiveConcertSessionTab(TSharedRef<IConcertServerSession>
 	: FConcertSessionTabBase(SyncServer)
 	, InspectedSession(MoveTemp(InspectedSession))
 	, ConstructUnderWindow(MoveTemp(ConstructUnderWindow))
-	, SessionHistoryController(MakeShared<FServerSessionHistoryController>(InspectedSession, SyncServer))
+	, SessionHistoryController(MakeShared<FLiveServerSessionHistoryController>(InspectedSession, SyncServer))
 	, PackageViewerController(MakeShared<FConcertSessionPackageViewerController>(InspectedSession, SyncServer))
 {}
 
diff --git a/Engine/Plugins/Developer/Concert/ConcertApp/MultiUserServer/Source/MultiUserServer/Private/Widgets/SessionTabs/Live/LiveConcertSessionTab.h b/Engine/Plugins/Developer/Concert/ConcertApp/MultiUserServer/Source/MultiUserServer/Private/Widgets/SessionTabs/Live/LiveConcertSessionTab.h
--- a/Engine/Plugins/Developer/Concert/ConcertApp/MultiUserServer/Source/MultiUserServer/Private/Widgets/SessionTabs/Live/LiveConcertSessionTab.h
+++ b/Engine/Plugins/Developer/Concert/ConcertApp/MultiUserServer/Source/MultiUserServer/Private/Widgets/SessionTabs/Live/LiveConcertSessionTab.h
@@ -3,12 +3,11 @@
 #pragma once
 
 #include "CoreMinimal.h"
-#include "Widgets/SessionTabs/Live/ServerSessionHistoryController.h"
 #include "Widgets/SessionTabs/ConcertSessionTabBase.h"
 
 class FConcertSessionPackageViewerController;
+class FLiveServerSessionHistoryController;
 class IConcertSyncServer;
-class FServerSessionHistoryController;
 class IConcertServerSession;
 class SDockTab;
 class SWindow;
@@ -40,7 +39,7 @@ private:
 	TAttribute<TSharedRef<SWindow>> ConstructUnderWindow;
 
 	/** Manages the session history widget */
-	const TSharedRef<FServerSessionHistoryController> SessionHistoryController;
+	const TSharedRef<FLiveServerSessionHistoryController> SessionHistoryController;
 	/** Manages the package viewer widget */
 	const TSharedRef<FConcertSessionPackageViewerController> PackageViewerController;
 };
diff --git a/Engine/Plugins/Developer/Concert/ConcertApp/MultiUserServer/Source/MultiUserServer/Private/Widgets/SessionTabs/Live/LiveServerSessionHistoryController.cpp b/Engine/Plugins/Developer/Concert/ConcertApp/MultiUserServer/Source/MultiUserServer/Private/Widgets/SessionTabs/Live/LiveServerSessionHistoryController.cpp
new file mode 100644
--- /dev/null
+++ b/Engine/Plugins/Developer/Concert/ConcertApp/MultiUserServer/Source/MultiUserServer/Private/Widgets/SessionTabs/Live/LiveServerSessionHistoryController.cpp
@@ -0,0 +1,37 @@
+﻿// Copyright Epic Games, Inc. All Rights Reserved.
+
+#include "LiveServerSessionHistoryController.h"
+
+#include "ConcertSyncSessionDatabase.h"
+#include "IConcertSession.h"
+#include "IConcertSyncServer.h"
+
+FLiveServerSessionHistoryController::FLiveServerSessionHistoryController(TSharedRef<IConcertServerSession> InspectedSession, TSharedRef<IConcertSyncServer> SyncServer)
+	: FServerSessionHistoryControllerBase(InspectedSession->GetId())
+	, SyncServer(MoveTemp(SyncServer))
+{
+	ReloadActivities();
+
+	if (const TOptional<FConcertSyncSessionDatabaseNonNullPtr> Database = SyncServer->GetLiveSessionDatabase(GetSessionId()))
+	{
+		Database->OnActivityProduced().AddRaw(this, &FLiveServerSessionHistoryController::OnSessionProduced);
+	}
+}
+
+FLiveServerSessionHistoryController::~FLiveServerSessionHistoryController()
+{
+	if (const TOptional<FConcertSyncSessionDatabaseNonNullPtr> Database = SyncServer->GetLiveSessionDatabase(GetSessionId()))
+	{
+		Database->OnActivityProduced().RemoveAll(this);
+	}
+}
+
+TOptional<FConcertSyncSessionDatabaseNonNullPtr> FLiveServerSessionHistoryController::GetSessionDatabase(const FGuid& InSessionId) const
+{
+	return SyncServer->GetLiveSessionDatabase(InSessionId);
+}
+
+void FLiveServerSessionHistoryController::OnSessionProduced(const FConcertSyncActivity& ProducedActivity)
+{
+	ReloadActivities();
+}
diff --git a/Engine/Plugins/Developer/Concert/ConcertApp/MultiUserServer/Source/MultiUserServer/Private/Widgets/SessionTabs/Live/LiveServerSessionHistoryController.h b/Engine/Plugins/Developer/Concert/ConcertApp/MultiUserServer/Source/MultiUserServer/Private/Widgets/SessionTabs/Live/LiveServerSessionHistoryController.h
new file mode 100644
--- /dev/null
+++ b/Engine/Plugins/Developer/Concert/ConcertApp/MultiUserServer/Source/MultiUserServer/Private/Widgets/SessionTabs/Live/LiveServerSessionHistoryController.h
@@ -0,0 +1,31 @@
+﻿// Copyright Epic Games, Inc. All Rights Reserved.
+
+#pragma once
+
+#include "CoreMinimal.h"
+#include "Widgets/SessionTabs/ServerSessionHistoryControllerBase.h"
+
+struct FConcertSyncActivity;
+class FConcertSyncSessionDatabase;
+class IConcertSyncServer;
+class IConcertServerSession;
+
+class FLiveServerSessionHistoryController : public FServerSessionHistoryControllerBase
+{
+public:
+
+	FLiveServerSessionHistoryController(TSharedRef<IConcertServerSession> InspectedSession, TSharedRef<IConcertSyncServer> SyncServer);
+	virtual ~FLiveServerSessionHistoryController() override;
+	
+protected:
+
+	//~ Begin FServerSessionHistoryControllerBase Interface
+	virtual TOptional<FConcertSyncSessionDatabaseNonNullPtr> GetSessionDatabase(const FGuid& InSessionId) const override;
+	//~ End FServerSessionHistoryControllerBase Interface
+
+private:
+
+	TSharedRef<IConcertSyncServer> SyncServer;
+	
+	void OnSessionProduced(const FConcertSyncActivity& ProducedActivity);
+};
diff --git a/Engine/Plugins/Developer/Concert/ConcertApp/MultiUserServer/Source/MultiUserServer/Private/Widgets/SessionTabs/Live/ServerSessionHistoryController.cpp b/Engine/Plugins/Developer/Concert/ConcertApp/MultiUserServer/Source/MultiUserServer/Private/Widgets/SessionTabs/Live/ServerSessionHistoryController.cpp
deleted file mode 100644
diff --git a/Engine/Plugins/Developer/Concert/ConcertApp/MultiUserServer/Source/MultiUserServer/Private/Widgets/SessionTabs/Live/ServerSessionHistoryController.h b/Engine/Plugins/Developer/Concert/ConcertApp/MultiUserServer/Source/MultiUserServer/Private/Widgets/SessionTabs/Live/ServerSessionHistoryController.h
deleted file mode 100644
diff --git a/Engine/Plugins/Developer/Concert/ConcertApp/MultiUserServer/Source/MultiUserServer/Private/Widgets/SessionTabs/ServerSessionHistoryControllerBase.cpp b/Engine/Plugins/Developer/Concert/ConcertApp/MultiUserServer/Source/MultiUserServer/Private/Widgets/SessionTabs/ServerSessionHistoryControllerBase.cpp
new file mode 100644
--- /dev/null
+++ b/Engine/Plugins/Developer/Concert/ConcertApp/MultiUserServer/Source/MultiUserServer/Private/Widgets/SessionTabs/ServerSessionHistoryControllerBase.cpp
@@ -0,0 +1,71 @@
+// Copyright Epic Games, Inc. All Rights Reserved.
+
+#include "ServerSessionHistoryControllerBase.h"
+
+#include "ConcertSyncSessionDatabase.h"
+#include "IConcertSession.h"
+#include "IConcertSyncServer.h"
+
+FServerSessionHistoryControllerBase::FServerSessionHistoryControllerBase(FGuid SessionId)
+	: SessionId(MoveTemp(SessionId))
+{}
+
+void FServerSessionHistoryControllerBase::GetActivities(int64 MaximumNumberOfActivities, TMap<FGuid, FConcertClientInfo>& OutEndpointClientInfoMap, TArray<FConcertSessionActivity>& OutFetchedActivities) const
+{
+	if (const TOptional<FConcertSyncSessionDatabaseNonNullPtr> Database = GetSessionDatabase(SessionId))
+	{
+		OutEndpointClientInfoMap.Reset();
+		OutFetchedActivities.Reset();
+
+		int64 LastActivityId = INDEX_NONE;
+		Database->GetActivityMaxId(LastActivityId);
+		const int64 FirstActivityIdToFetch = FMath::Max<int64>(1, LastActivityId - MaximumNumberOfActivities);
+		Database->EnumerateActivitiesInRange(FirstActivityIdToFetch, MaximumNumberOfActivities, [this, &Database, &OutEndpointClientInfoMap, &OutFetchedActivities](FConcertSyncActivity&& InActivity)
+		{
+			if (!OutEndpointClientInfoMap.Contains(InActivity.EndpointId))
+			{
+				FConcertSyncEndpointData EndpointData;
+				if (Database->GetEndpoint(InActivity.EndpointId, EndpointData))
+				{
+					OutEndpointClientInfoMap.Add(InActivity.EndpointId, EndpointData.ClientInfo);
+				}
+			}
+
+			FStructOnScope ActivitySummary;
+			if (InActivity.EventSummary.GetPayload(ActivitySummary))
+			{
+				OutFetchedActivities.Emplace(MoveTemp(InActivity), MoveTemp(ActivitySummary));
+			}
+
+			return EBreakBehavior::Continue;
+		});
+	}
+}
+
+bool FServerSessionHistoryControllerBase::GetPackageEvent(const FConcertSessionActivity& Activity, FConcertSyncPackageEventMetaData& OutPackageEvent) const
+{
+	if (const TOptional<FConcertSyncSessionDatabaseNonNullPtr> Database = GetSessionDatabase(SessionId))
+	{
+		return Database->GetPackageEventMetaData(Activity.Activity.EventId, OutPackageEvent.PackageRevision, OutPackageEvent.PackageInfo);
+	}
+	
+	return false;
+}
+
+TFuture<TOptional<FConcertSyncTransactionEvent>> FServerSessionHistoryControllerBase::GetTransactionEvent(const FConcertSessionActivity& Activity) const
+{
+	if (const TOptional<FConcertSyncSessionDatabaseNonNullPtr> Database = GetSessionDatabase(SessionId))
+	{
+		return FindOrRequestTransactionEvent(*Database, Activity.Activity.EventId);
+	}
+	
+	return MakeFulfilledPromise<TOptional<FConcertSyncTransactionEvent>>().GetFuture(); // Not found.
+}
+
+TFuture<TOptional<FConcertSyncTransactionEvent>> FServerSessionHistoryControllerBase::FindOrRequestTransactionEvent(const FConcertSyncSessionDatabase& Database, const int64 TransactionEventId) const
+{
+	FConcertSyncTransactionEvent TransactionEvent;
+	return Database.GetTransactionEvent(TransactionEventId, TransactionEvent, false)
+		?	MakeFulfilledPromise<TOptional<FConcertSyncTransactionEvent>>(MoveTemp(TransactionEvent)).GetFuture()
+		:	MakeFulfilledPromise<TOptional<FConcertSyncTransactionEvent>>().GetFuture();
+}
diff --git a/Engine/Plugins/Developer/Concert/ConcertApp/MultiUserServer/Source/MultiUserServer/Private/Widgets/SessionTabs/ServerSessionHistoryControllerBase.h b/Engine/Plugins/Developer/Concert/ConcertApp/MultiUserServer/Source/MultiUserServer/Private/Widgets/SessionTabs/ServerSessionHistoryControllerBase.h
new file mode 100644
--- /dev/null
+++ b/Engine/Plugins/Developer/Concert/ConcertApp/MultiUserServer/Source/MultiUserServer/Private/Widgets/SessionTabs/ServerSessionHistoryControllerBase.h
@@ -0,0 +1,43 @@
+// Copyright Epic Games, Inc. All Rights Reserved.
+
+#pragma once
+
+#include "CoreMinimal.h"
+#include "IConcertSyncServer.h"
+#include "Session/History/AbstractSessionHistoryController.h"
+
+struct FConcertSyncActivity;
+class FConcertSyncSessionDatabase;
+class IConcertSyncServer;
+class IConcertServerSession;
+
+/**
+ * Abstract base class providing most functionality.
+ * Subclasses control which session database is used: archived or live.
+ */
+class FServerSessionHistoryControllerBase : public FAbstractSessionHistoryController
+{
+public:
+
+	FServerSessionHistoryControllerBase(FGuid SessionId);
+	
+protected:
+
+	//~ Begin FAbstractSessionHistoryController Interface
+	virtual void GetActivities(int64 MaximumNumberOfActivities, TMap<FGuid, FConcertClientInfo>& OutEndpointClientInfoMap, TArray<FConcertSessionActivity>& OutFetchedActivities) const override;
+	virtual bool GetPackageEvent(const FConcertSessionActivity& Activity, FConcertSyncPackageEventMetaData& OutPackageEvent) const override;
+	virtual TFuture<TOptional<FConcertSyncTransactionEvent>> GetTransactionEvent(const FConcertSessionActivity& Activity) const override;
+	//~ End FAbstractSessionHistoryController Interface
+
+	/** Subclasses control which types of transaction are displayed: archived or live */
+	virtual TOptional<FConcertSyncSessionDatabaseNonNullPtr> GetSessionDatabase(const FGuid& InSessionId) const = 0;
+
+	const FGuid& GetSessionId() const { return SessionId; }
+	
+private:
+
+	/** The ID of the managed session */
+	FGuid SessionId;
+
+	TFuture<TOptional<FConcertSyncTransactionEvent>> FindOrRequestTransactionEvent(const FConcertSyncSessionDatabase& Database, const int64 TransactionEventId) const;
+};
