Create session overview for archived sessions

diff --git a/Engine/Plugins/Developer/Concert/ConcertApp/MultiUserServer/Source/MultiUserServer/Private/Widgets/SessionTabs/Archived/ArchivedConcertSessionTab.cpp b/Engine/Plugins/Developer/Concert/ConcertApp/MultiUserServer/Source/MultiUserServer/Private/Widgets/SessionTabs/Archived/ArchivedConcertSessionTab.cpp
--- a/Engine/Plugins/Developer/Concert/ConcertApp/MultiUserServer/Source/MultiUserServer/Private/Widgets/SessionTabs/Archived/ArchivedConcertSessionTab.cpp
+++ b/Engine/Plugins/Developer/Concert/ConcertApp/MultiUserServer/Source/MultiUserServer/Private/Widgets/SessionTabs/Archived/ArchivedConcertSessionTab.cpp
@@ -2,13 +2,16 @@
 
 #include "Widgets/SessionTabs/Archived/ArchivedConcertSessionTab.h"
 
+#include "ArchivedSessionHistoryController.h"
+#include "Session/History/SSessionHistory.h"
 #include "Widgets/SessionTabs/Archived/SConcertArchivedSessionInspector.h"
 #include "Widgets/StatusBar/SConcertStatusBar.h"
 
-FArchivedConcertSessionTab::FArchivedConcertSessionTab(const FGuid& InspectedSessionID, TSharedRef<IConcertSyncServer> SyncServer, TAttribute<TSharedRef<SWindow>> ConstructUnderWindow)
+FArchivedConcertSessionTab::FArchivedConcertSessionTab(FGuid InspectedSessionID, TSharedRef<IConcertSyncServer> SyncServer, TAttribute<TSharedRef<SWindow>> ConstructUnderWindow)
 	: FConcertSessionTabBase(SyncServer)
-	, InspectedSessionID(InspectedSessionID)
-	, ConstructUnderWindow(ConstructUnderWindow)
+	, InspectedSessionID(MoveTemp(InspectedSessionID))
+	, ConstructUnderWindow(MoveTemp(ConstructUnderWindow))
+	, HistoryController(MakeShared<FArchivedSessionHistoryController>(MoveTemp(InspectedSessionID), MoveTemp(SyncServer)))
 {}
 
 FGuid FArchivedConcertSessionTab::GetSessionID() const
@@ -19,7 +22,7 @@ FGuid FArchivedConcertSessionTab::GetSessionID() const
 void FArchivedConcertSessionTab::CreateDockContent(const TSharedRef<SDockTab>& InDockTab)
 {
 	InDockTab->SetContent(
-		SNew(SConcertArchivedSessionInspector, SConcertArchivedSessionInspector::FRequiredArgs(InDockTab, ConstructUnderWindow.Get()))
+		SNew(SConcertArchivedSessionInspector, SConcertArchivedSessionInspector::FRequiredArgs(InDockTab, ConstructUnderWindow.Get(), HistoryController->GetSessionHistory()))
 			.StatusBar()
 			[
 				SNew(SConcertStatusBar, *GetTabId())
diff --git a/Engine/Plugins/Developer/Concert/ConcertApp/MultiUserServer/Source/MultiUserServer/Private/Widgets/SessionTabs/Archived/ArchivedConcertSessionTab.h b/Engine/Plugins/Developer/Concert/ConcertApp/MultiUserServer/Source/MultiUserServer/Private/Widgets/SessionTabs/Archived/ArchivedConcertSessionTab.h
--- a/Engine/Plugins/Developer/Concert/ConcertApp/MultiUserServer/Source/MultiUserServer/Private/Widgets/SessionTabs/Archived/ArchivedConcertSessionTab.h
+++ b/Engine/Plugins/Developer/Concert/ConcertApp/MultiUserServer/Source/MultiUserServer/Private/Widgets/SessionTabs/Archived/ArchivedConcertSessionTab.h
@@ -5,6 +5,7 @@
 #include "CoreMinimal.h"
 #include "Widgets/SessionTabs/ConcertSessionTabBase.h"
 
+class FArchivedSessionHistoryController;
 class SConcertArchivedSessionInspector;
 
 /** Manages the tab for an archived session.  */
@@ -12,7 +13,7 @@ class FArchivedConcertSessionTab : public FConcertSessionTabBase
 {
 public:
 
-	FArchivedConcertSessionTab(const FGuid& InspectedSessionID, TSharedRef<IConcertSyncServer> SyncServer, TAttribute<TSharedRef<SWindow>> ConstructUnderWindow);
+	FArchivedConcertSessionTab(FGuid InspectedSessionID, TSharedRef<IConcertSyncServer> SyncServer, TAttribute<TSharedRef<SWindow>> ConstructUnderWindow);
 
 protected:
 
@@ -25,11 +26,13 @@ protected:
 private:
 
 	/** The inspected session's ID */
-	FGuid InspectedSessionID;
+	const FGuid InspectedSessionID;
 
 	/** Used later to obtain the window into which to add the tab */
-	TAttribute<TSharedRef<SWindow>> ConstructUnderWindow;
-    
+	const TAttribute<TSharedRef<SWindow>> ConstructUnderWindow;
+
+	const TSharedRef<FArchivedSessionHistoryController> HistoryController;
+	
 	/** Displays session */
 	TSharedPtr<SConcertArchivedSessionInspector> Inspector;
 };
diff --git a/Engine/Plugins/Developer/Concert/ConcertApp/MultiUserServer/Source/MultiUserServer/Private/Widgets/SessionTabs/Archived/ArchivedSessionHistoryController.cpp b/Engine/Plugins/Developer/Concert/ConcertApp/MultiUserServer/Source/MultiUserServer/Private/Widgets/SessionTabs/Archived/ArchivedSessionHistoryController.cpp
new file mode 100644
--- /dev/null
+++ b/Engine/Plugins/Developer/Concert/ConcertApp/MultiUserServer/Source/MultiUserServer/Private/Widgets/SessionTabs/Archived/ArchivedSessionHistoryController.cpp
@@ -0,0 +1,13 @@
+﻿// Copyright Epic Games, Inc. All Rights Reserved.
+
+#include "ArchivedSessionHistoryController.h"
+
+FArchivedSessionHistoryController::FArchivedSessionHistoryController(FGuid SessionId, TSharedRef<IConcertSyncServer> SyncServer)
+	: FServerSessionHistoryControllerBase(MoveTemp(SessionId))
+	, SyncServer(MoveTemp(SyncServer))
+{}
+
+TOptional<FConcertSyncSessionDatabaseNonNullPtr> FArchivedSessionHistoryController::GetSessionDatabase(const FGuid& InSessionId) const
+{
+	return SyncServer->GetLiveSessionDatabase(InSessionId);
+}
diff --git a/Engine/Plugins/Developer/Concert/ConcertApp/MultiUserServer/Source/MultiUserServer/Private/Widgets/SessionTabs/Archived/ArchivedSessionHistoryController.h b/Engine/Plugins/Developer/Concert/ConcertApp/MultiUserServer/Source/MultiUserServer/Private/Widgets/SessionTabs/Archived/ArchivedSessionHistoryController.h
new file mode 100644
--- /dev/null
+++ b/Engine/Plugins/Developer/Concert/ConcertApp/MultiUserServer/Source/MultiUserServer/Private/Widgets/SessionTabs/Archived/ArchivedSessionHistoryController.h
@@ -0,0 +1,23 @@
+﻿// Copyright Epic Games, Inc. All Rights Reserved.
+
+#pragma once
+
+#include "CoreMinimal.h"
+#include "Widgets/SessionTabs/ServerSessionHistoryControllerBase.h"
+
+class FArchivedSessionHistoryController : public FServerSessionHistoryControllerBase
+{
+public:
+	
+	FArchivedSessionHistoryController(FGuid SessionId, TSharedRef<IConcertSyncServer> SyncServer);
+	
+protected:
+
+	//~ Begin FServerSessionHistoryControllerBase Interface
+	virtual TOptional<FConcertSyncSessionDatabaseNonNullPtr> GetSessionDatabase(const FGuid& InSessionId) const override;
+	//~ End FServerSessionHistoryControllerBase Interface
+
+private:
+
+	const TSharedRef<IConcertSyncServer> SyncServer;
+};
diff --git a/Engine/Plugins/Developer/Concert/ConcertApp/MultiUserServer/Source/MultiUserServer/Private/Widgets/SessionTabs/Archived/SConcertArchivedSessionInspector.cpp b/Engine/Plugins/Developer/Concert/ConcertApp/MultiUserServer/Source/MultiUserServer/Private/Widgets/SessionTabs/Archived/SConcertArchivedSessionInspector.cpp
--- a/Engine/Plugins/Developer/Concert/ConcertApp/MultiUserServer/Source/MultiUserServer/Private/Widgets/SessionTabs/Archived/SConcertArchivedSessionInspector.cpp
+++ b/Engine/Plugins/Developer/Concert/ConcertApp/MultiUserServer/Source/MultiUserServer/Private/Widgets/SessionTabs/Archived/SConcertArchivedSessionInspector.cpp
@@ -7,7 +7,6 @@
 #include "Widgets/Docking/SDockTab.h"
 #include "Widgets/Layout/SBorder.h"
 #include "Widgets/SBoxPanel.h"
-#include "Widgets/SNullWidget.h"
 
 #define LOCTEXT_NAMESPACE "SConcertArchivedSessionInspector"
 
@@ -23,13 +22,6 @@ void SConcertArchivedSessionInspector::Construct(const FArguments& InArgs, const
 		[
 			SNew(SVerticalBox)
 
-			// Toolbar
-			+SVerticalBox::Slot()
-			.AutoHeight()
-			[
-				SNullWidget::NullWidget
-			]
-
 			// Content
 			+SVerticalBox::Slot()
 			.FillHeight(1.f)
@@ -57,7 +49,7 @@ TSharedRef<SWidget> SConcertArchivedSessionInspector::CreateTabs(const FRequired
 {
 	TabManager = FGlobalTabmanager::Get()->NewTabManager(RequiredArgs.ConstructUnderMajorTab);
 	
-	TabManager->RegisterTabSpawner(HistoryTabId, FOnSpawnTab::CreateSP(this, &SConcertArchivedSessionInspector::SpawnActivityHistory))
+	TabManager->RegisterTabSpawner(HistoryTabId, FOnSpawnTab::CreateSP(this, &SConcertArchivedSessionInspector::SpawnActivityHistory, RequiredArgs.SessionHistory))
 		.SetDisplayName(LOCTEXT("ActivityHistoryLabel", "History"));
 
 	const TSharedRef<FTabManager::FLayout> Layout = FTabManager::NewLayout("ConcertArchivedSessionLayout_v0.1")
@@ -76,13 +68,13 @@ TSharedRef<SWidget> SConcertArchivedSessionInspector::CreateTabs(const FRequired
 	return TabManager->RestoreFrom(Layout, RequiredArgs.ConstructUnderWindow).ToSharedRef();
 }
 
-TSharedRef<SDockTab> SConcertArchivedSessionInspector::SpawnActivityHistory(const FSpawnTabArgs& Args)
+TSharedRef<SDockTab> SConcertArchivedSessionInspector::SpawnActivityHistory(const FSpawnTabArgs& Args, TSharedRef<SSessionHistory> SessionHistory)
 {
 	return SNew(SDockTab)
 		.Label(LOCTEXT("ActivityHistoryLabel", "History"))
 		.TabRole(PanelTab)
 		[
-			SNullWidget::NullWidget
+			SessionHistory
 		]; 
 }
 
diff --git a/Engine/Plugins/Developer/Concert/ConcertApp/MultiUserServer/Source/MultiUserServer/Private/Widgets/SessionTabs/Archived/SConcertArchivedSessionInspector.h b/Engine/Plugins/Developer/Concert/ConcertApp/MultiUserServer/Source/MultiUserServer/Private/Widgets/SessionTabs/Archived/SConcertArchivedSessionInspector.h
--- a/Engine/Plugins/Developer/Concert/ConcertApp/MultiUserServer/Source/MultiUserServer/Private/Widgets/SessionTabs/Archived/SConcertArchivedSessionInspector.h
+++ b/Engine/Plugins/Developer/Concert/ConcertApp/MultiUserServer/Source/MultiUserServer/Private/Widgets/SessionTabs/Archived/SConcertArchivedSessionInspector.h
@@ -7,6 +7,8 @@
 #include "Widgets/Docking/SDockTab.h"
 #include "Widgets/SCompoundWidget.h"
 
+class SSessionHistory;
+
 /** Designed as content for a tab. Displays information about an archived session. */
 class SConcertArchivedSessionInspector : public SCompoundWidget
 {
@@ -16,10 +18,12 @@ public:
 	{
 		TSharedRef<SDockTab> ConstructUnderMajorTab;
 		TSharedRef<SWindow> ConstructUnderWindow;
+		TSharedRef<SSessionHistory> SessionHistory;
 
-		FRequiredArgs(TSharedRef<SDockTab> ConstructUnderMajorTab, TSharedRef<SWindow> ConstructUnderWindow)
+		FRequiredArgs(TSharedRef<SDockTab> ConstructUnderMajorTab, TSharedRef<SWindow> ConstructUnderWindow, TSharedRef<SSessionHistory> SessionHistory)
 			: ConstructUnderMajorTab(MoveTemp(ConstructUnderMajorTab))
 			, ConstructUnderWindow(MoveTemp(ConstructUnderWindow))
+			, SessionHistory(SessionHistory)
 		{}
 	};
 	
@@ -38,5 +42,5 @@ private:
 	TSharedPtr<FTabManager> TabManager;
 	
 	TSharedRef<SWidget> CreateTabs(const FRequiredArgs& RequiredArgs);
-	TSharedRef<SDockTab> SpawnActivityHistory(const FSpawnTabArgs& Args);
+	TSharedRef<SDockTab> SpawnActivityHistory(const FSpawnTabArgs& Args, TSharedRef<SSessionHistory> SessionHistory);
 };
