Create UI for deleting activities in archived session

diff --git a/Engine/Plugins/Developer/Concert/ConcertApp/MultiUserClient/Source/MultiUserClient/Private/Widgets/ClientSessionHistoryController.cpp b/Engine/Plugins/Developer/Concert/ConcertApp/MultiUserClient/Source/MultiUserClient/Private/Widgets/ClientSessionHistoryController.cpp
--- a/Engine/Plugins/Developer/Concert/ConcertApp/MultiUserClient/Source/MultiUserClient/Private/Widgets/ClientSessionHistoryController.cpp
+++ b/Engine/Plugins/Developer/Concert/ConcertApp/MultiUserClient/Source/MultiUserClient/Private/Widgets/ClientSessionHistoryController.cpp
@@ -9,7 +9,7 @@
 #include "Session/History/SSessionHistory.h"
 
 FClientSessionHistoryController::FClientSessionHistoryController(TSharedRef<IConcertSyncClient> Client, FName PackageFilter)
-	: FAbstractSessionHistoryController(PackageFilter)
+	: FAbstractSessionHistoryController(SSessionHistory::FArguments().PackageFilter(PackageFilter))
 	, Client(MoveTemp(Client))
 {
 	SubscribeToWorkspaceEvents();
diff --git a/Engine/Plugins/Developer/Concert/ConcertApp/MultiUserServer/Source/MultiUserServer/Private/Widgets/SessionTabs/Archived/ArchivedConcertSessionTab.cpp b/Engine/Plugins/Developer/Concert/ConcertApp/MultiUserServer/Source/MultiUserServer/Private/Widgets/SessionTabs/Archived/ArchivedConcertSessionTab.cpp
--- a/Engine/Plugins/Developer/Concert/ConcertApp/MultiUserServer/Source/MultiUserServer/Private/Widgets/SessionTabs/Archived/ArchivedConcertSessionTab.cpp
+++ b/Engine/Plugins/Developer/Concert/ConcertApp/MultiUserServer/Source/MultiUserServer/Private/Widgets/SessionTabs/Archived/ArchivedConcertSessionTab.cpp
@@ -3,6 +3,7 @@
 #include "Widgets/SessionTabs/Archived/ArchivedConcertSessionTab.h"
 
 #include "ArchivedSessionHistoryController.h"
+#include "Session/History/SEditableSessionHistory.h"
 #include "Session/History/SSessionHistory.h"
 #include "Widgets/SessionTabs/Archived/SConcertArchivedSessionInspector.h"
 #include "Widgets/StatusBar/SConcertStatusBar.h"
@@ -10,8 +11,8 @@
 FArchivedConcertSessionTab::FArchivedConcertSessionTab(FGuid InspectedSessionID, TSharedRef<IConcertSyncServer> SyncServer, TAttribute<TSharedRef<SWindow>> ConstructUnderWindow)
 	: FConcertSessionTabBase(SyncServer)
 	, InspectedSessionID(MoveTemp(InspectedSessionID))
+	, SyncServer(MoveTemp(SyncServer))
 	, ConstructUnderWindow(MoveTemp(ConstructUnderWindow))
-	, HistoryController(MakeShared<FArchivedSessionHistoryController>(MoveTemp(InspectedSessionID), MoveTemp(SyncServer)))
 {}
 
 FGuid FArchivedConcertSessionTab::GetSessionID() const
@@ -21,8 +22,20 @@ FGuid FArchivedConcertSessionTab::GetSessionID() const
 
 void FArchivedConcertSessionTab::CreateDockContent(const TSharedRef<SDockTab>& InDockTab)
 {
+	SEditableSessionHistory::FMakeSessionHistory MakeSessionHistory = SEditableSessionHistory::FMakeSessionHistory::CreateLambda([this](SSessionHistory::FArguments Arguments)
+	{
+		checkf(!HistoryController.IsValid(), TEXT("Called more than once"));
+		HistoryController = MakeShared<FArchivedSessionHistoryController>(InspectedSessionID, SyncServer, MoveTemp(Arguments));
+		return HistoryController->GetSessionHistory();
+	});
+	
 	InDockTab->SetContent(
-		SNew(SConcertArchivedSessionInspector, SConcertArchivedSessionInspector::FRequiredArgs(InDockTab, ConstructUnderWindow.Get(), HistoryController->GetSessionHistory()))
+		SNew(SConcertArchivedSessionInspector)
+			.ConstructUnderMajorTab(InDockTab)
+			.ConstructUnderWindow(ConstructUnderWindow.Get())
+			.MakeSessionHistory(MoveTemp(MakeSessionHistory))
+			.DeleteActivity_Raw(this, &FArchivedConcertSessionTab::OnRequestDeleteActivity)
+			.CanDeleteActivity_Raw(this, &FArchivedConcertSessionTab::CanDeleteActivity)
 			.StatusBar()
 			[
 				SNew(SConcertStatusBar, *GetTabId())
@@ -30,5 +43,13 @@ void FArchivedConcertSessionTab::CreateDockContent(const TSharedRef<SDockTab>& I
 		);
 }
 
-void FArchivedConcertSessionTab::OnOpenTab()
-{}
+void FArchivedConcertSessionTab::OnRequestDeleteActivity(const TSharedRef<FConcertSessionActivity>& DeleteActivity) const
+{
+	// TODO:
+}
+
+bool FArchivedConcertSessionTab::CanDeleteActivity(const TSharedRef<FConcertSessionActivity>& DeleteActivity) const
+{
+	// TODO:
+	return false;
+}
diff --git a/Engine/Plugins/Developer/Concert/ConcertApp/MultiUserServer/Source/MultiUserServer/Private/Widgets/SessionTabs/Archived/ArchivedConcertSessionTab.h b/Engine/Plugins/Developer/Concert/ConcertApp/MultiUserServer/Source/MultiUserServer/Private/Widgets/SessionTabs/Archived/ArchivedConcertSessionTab.h
--- a/Engine/Plugins/Developer/Concert/ConcertApp/MultiUserServer/Source/MultiUserServer/Private/Widgets/SessionTabs/Archived/ArchivedConcertSessionTab.h
+++ b/Engine/Plugins/Developer/Concert/ConcertApp/MultiUserServer/Source/MultiUserServer/Private/Widgets/SessionTabs/Archived/ArchivedConcertSessionTab.h
@@ -5,6 +5,7 @@
 #include "CoreMinimal.h"
 #include "Widgets/SessionTabs/ConcertSessionTabBase.h"
 
+struct FConcertSessionActivity;
 class FArchivedSessionHistoryController;
 class SConcertArchivedSessionInspector;
 
@@ -20,7 +21,7 @@ protected:
 	//~ Begin FAbstractConcertSessionTab Interface
 	virtual FGuid GetSessionID() const override;
 	virtual void CreateDockContent(const TSharedRef<SDockTab>& InDockTab) override;
-	virtual void OnOpenTab() override;
+	virtual void OnOpenTab() override {}
 	//~ End FAbstractConcertSessionTab Interface
 	
 private:
@@ -28,11 +29,17 @@ private:
 	/** The inspected session's ID */
 	const FGuid InspectedSessionID;
 
+	/** Used later to construct Inspector */
+	const TSharedRef<IConcertSyncServer> SyncServer;
+	
 	/** Used later to obtain the window into which to add the tab */
 	const TAttribute<TSharedRef<SWindow>> ConstructUnderWindow;
-
-	const TSharedRef<FArchivedSessionHistoryController> HistoryController;
+	
+	TSharedPtr<FArchivedSessionHistoryController> HistoryController;
 	
 	/** Displays session */
 	TSharedPtr<SConcertArchivedSessionInspector> Inspector;
+
+	void OnRequestDeleteActivity(const TSharedRef<FConcertSessionActivity>& DeleteActivity) const;
+	bool CanDeleteActivity(const TSharedRef<FConcertSessionActivity>& DeleteActivity) const;
 };
diff --git a/Engine/Plugins/Developer/Concert/ConcertApp/MultiUserServer/Source/MultiUserServer/Private/Widgets/SessionTabs/Archived/ArchivedSessionHistoryController.cpp b/Engine/Plugins/Developer/Concert/ConcertApp/MultiUserServer/Source/MultiUserServer/Private/Widgets/SessionTabs/Archived/ArchivedSessionHistoryController.cpp
--- a/Engine/Plugins/Developer/Concert/ConcertApp/MultiUserServer/Source/MultiUserServer/Private/Widgets/SessionTabs/Archived/ArchivedSessionHistoryController.cpp
+++ b/Engine/Plugins/Developer/Concert/ConcertApp/MultiUserServer/Source/MultiUserServer/Private/Widgets/SessionTabs/Archived/ArchivedSessionHistoryController.cpp
@@ -2,12 +2,14 @@
 
 #include "ArchivedSessionHistoryController.h"
 
-FArchivedSessionHistoryController::FArchivedSessionHistoryController(FGuid SessionId, TSharedRef<IConcertSyncServer> SyncServer)
-	: FServerSessionHistoryControllerBase(MoveTemp(SessionId))
+FArchivedSessionHistoryController::FArchivedSessionHistoryController(FGuid SessionId, TSharedRef<IConcertSyncServer> SyncServer, SSessionHistory::FArguments Arguments)
+	: FServerSessionHistoryControllerBase(MoveTemp(SessionId), MoveTemp(Arguments))
 	, SyncServer(MoveTemp(SyncServer))
-{}
+{
+	ReloadActivities();
+}
 
 TOptional<FConcertSyncSessionDatabaseNonNullPtr> FArchivedSessionHistoryController::GetSessionDatabase(const FGuid& InSessionId) const
 {
-	return SyncServer->GetLiveSessionDatabase(InSessionId);
+	return SyncServer->GetArchivedSessionDatabase(InSessionId);
 }
diff --git a/Engine/Plugins/Developer/Concert/ConcertApp/MultiUserServer/Source/MultiUserServer/Private/Widgets/SessionTabs/Archived/ArchivedSessionHistoryController.h b/Engine/Plugins/Developer/Concert/ConcertApp/MultiUserServer/Source/MultiUserServer/Private/Widgets/SessionTabs/Archived/ArchivedSessionHistoryController.h
--- a/Engine/Plugins/Developer/Concert/ConcertApp/MultiUserServer/Source/MultiUserServer/Private/Widgets/SessionTabs/Archived/ArchivedSessionHistoryController.h
+++ b/Engine/Plugins/Developer/Concert/ConcertApp/MultiUserServer/Source/MultiUserServer/Private/Widgets/SessionTabs/Archived/ArchivedSessionHistoryController.h
@@ -9,7 +9,7 @@ class FArchivedSessionHistoryController : public FServerSessionHistoryController
 {
 public:
 	
-	FArchivedSessionHistoryController(FGuid SessionId, TSharedRef<IConcertSyncServer> SyncServer);
+	FArchivedSessionHistoryController(FGuid SessionId, TSharedRef<IConcertSyncServer> SyncServer, SSessionHistory::FArguments Arguments);
 	
 protected:
 
diff --git a/Engine/Plugins/Developer/Concert/ConcertApp/MultiUserServer/Source/MultiUserServer/Private/Widgets/SessionTabs/Archived/SConcertArchivedSessionInspector.cpp b/Engine/Plugins/Developer/Concert/ConcertApp/MultiUserServer/Source/MultiUserServer/Private/Widgets/SessionTabs/Archived/SConcertArchivedSessionInspector.cpp
--- a/Engine/Plugins/Developer/Concert/ConcertApp/MultiUserServer/Source/MultiUserServer/Private/Widgets/SessionTabs/Archived/SConcertArchivedSessionInspector.cpp
+++ b/Engine/Plugins/Developer/Concert/ConcertApp/MultiUserServer/Source/MultiUserServer/Private/Widgets/SessionTabs/Archived/SConcertArchivedSessionInspector.cpp
@@ -2,6 +2,7 @@
 
 #include "SConcertArchivedSessionInspector.h"
 
+#include "EditorStyleSet.h"
 #include "Framework/Docking/TabManager.h"
 #include "Session/History/SSessionHistory.h"
 #include "Widgets/Docking/SDockTab.h"
@@ -12,8 +13,10 @@
 
 const FName SConcertArchivedSessionInspector::HistoryTabId("HistoryTabId");
 
-void SConcertArchivedSessionInspector::Construct(const FArguments& InArgs, const FRequiredArgs& InRequiredArgs)
+void SConcertArchivedSessionInspector::Construct(const FArguments& InArgs)
 {
+	check(InArgs._ConstructUnderWindow && InArgs._ConstructUnderMajorTab && InArgs._MakeSessionHistory.IsBound() && InArgs._CanDeleteActivity.IsBound());
+	
 	ChildSlot
 	[
 		SNew(SBorder)
@@ -30,7 +33,7 @@ void SConcertArchivedSessionInspector::Construct(const FArguments& InArgs, const
 				.BorderImage(FAppStyle::Get().GetBrush("Brushes.Title"))
 				.Padding(FMargin(0.f, 0.f, 0.f, 5.0f)) // Visually separate status bar from tabs
 				[
-					CreateTabs(InRequiredArgs)
+					CreateTabs(InArgs)
 				]
 			]
 
@@ -45,11 +48,11 @@ void SConcertArchivedSessionInspector::Construct(const FArguments& InArgs, const
 	];
 }
 
-TSharedRef<SWidget> SConcertArchivedSessionInspector::CreateTabs(const FRequiredArgs& RequiredArgs)
+TSharedRef<SWidget> SConcertArchivedSessionInspector::CreateTabs(const FArguments& InArgs)
 {
-	TabManager = FGlobalTabmanager::Get()->NewTabManager(RequiredArgs.ConstructUnderMajorTab);
+	TabManager = FGlobalTabmanager::Get()->NewTabManager(InArgs._ConstructUnderMajorTab.ToSharedRef());
 	
-	TabManager->RegisterTabSpawner(HistoryTabId, FOnSpawnTab::CreateSP(this, &SConcertArchivedSessionInspector::SpawnActivityHistory, RequiredArgs.SessionHistory))
+	TabManager->RegisterTabSpawner(HistoryTabId, FOnSpawnTab::CreateSP(this, &SConcertArchivedSessionInspector::SpawnActivityHistory, InArgs._MakeSessionHistory, InArgs._CanDeleteActivity, InArgs._DeleteActivity))
 		.SetDisplayName(LOCTEXT("ActivityHistoryLabel", "History"));
 
 	const TSharedRef<FTabManager::FLayout> Layout = FTabManager::NewLayout("ConcertArchivedSessionLayout_v0.1")
@@ -65,16 +68,24 @@ TSharedRef<SWidget> SConcertArchivedSessionInspector::CreateTabs(const FRequired
 				)
 		);
 
-	return TabManager->RestoreFrom(Layout, RequiredArgs.ConstructUnderWindow).ToSharedRef();
+	return TabManager->RestoreFrom(Layout, InArgs._ConstructUnderWindow).ToSharedRef();
 }
 
-TSharedRef<SDockTab> SConcertArchivedSessionInspector::SpawnActivityHistory(const FSpawnTabArgs& Args, TSharedRef<SSessionHistory> SessionHistory)
+TSharedRef<SDockTab> SConcertArchivedSessionInspector::SpawnActivityHistory(
+	const FSpawnTabArgs& Args,
+	SEditableSessionHistory::FMakeSessionHistory MakeSessionHistory,
+	SEditableSessionHistory::FCanDeleteActivity CanDeleteActivity,
+	SEditableSessionHistory::FRequestDeleteActivity DeleteActivity)
 {
+	SessionHistory = SNew(SEditableSessionHistory)
+		.MakeSessionHistory(MoveTemp(MakeSessionHistory))
+		.CanDeleteActivity(MoveTemp(CanDeleteActivity))
+		.DeleteActivity(MoveTemp(DeleteActivity));
 	return SNew(SDockTab)
 		.Label(LOCTEXT("ActivityHistoryLabel", "History"))
 		.TabRole(PanelTab)
 		[
-			SessionHistory
+			SessionHistory.ToSharedRef()
 		]; 
 }
 
diff --git a/Engine/Plugins/Developer/Concert/ConcertApp/MultiUserServer/Source/MultiUserServer/Private/Widgets/SessionTabs/Archived/SConcertArchivedSessionInspector.h b/Engine/Plugins/Developer/Concert/ConcertApp/MultiUserServer/Source/MultiUserServer/Private/Widgets/SessionTabs/Archived/SConcertArchivedSessionInspector.h
--- a/Engine/Plugins/Developer/Concert/ConcertApp/MultiUserServer/Source/MultiUserServer/Private/Widgets/SessionTabs/Archived/SConcertArchivedSessionInspector.h
+++ b/Engine/Plugins/Developer/Concert/ConcertApp/MultiUserServer/Source/MultiUserServer/Private/Widgets/SessionTabs/Archived/SConcertArchivedSessionInspector.h
@@ -3,44 +3,45 @@
 #pragma once
 
 #include "CoreMinimal.h"
+
+#include "Session/History/SEditableSessionHistory.h"
+#include "Session/History/SSessionHistory.h"
 #include "Widgets/DeclarativeSyntaxSupport.h"
-#include "Widgets/Docking/SDockTab.h"
 #include "Widgets/SCompoundWidget.h"
 
+class SDockTab;
 class SSessionHistory;
 
 /** Designed as content for a tab. Displays information about an archived session. */
 class SConcertArchivedSessionInspector : public SCompoundWidget
 {
 public:
-
-	struct FRequiredArgs
-	{
-		TSharedRef<SDockTab> ConstructUnderMajorTab;
-		TSharedRef<SWindow> ConstructUnderWindow;
-		TSharedRef<SSessionHistory> SessionHistory;
-
-		FRequiredArgs(TSharedRef<SDockTab> ConstructUnderMajorTab, TSharedRef<SWindow> ConstructUnderWindow, TSharedRef<SSessionHistory> SessionHistory)
-			: ConstructUnderMajorTab(MoveTemp(ConstructUnderMajorTab))
-			, ConstructUnderWindow(MoveTemp(ConstructUnderWindow))
-			, SessionHistory(SessionHistory)
-		{}
-	};
 	
 	static const FName HistoryTabId;
 
 	SLATE_BEGIN_ARGS(SConcertArchivedSessionInspector)
 	{}
+		SLATE_ARGUMENT(TSharedPtr<SDockTab>, ConstructUnderMajorTab)
+		SLATE_ARGUMENT(TSharedPtr<SWindow>, ConstructUnderWindow)
+		SLATE_EVENT(SEditableSessionHistory::FMakeSessionHistory, MakeSessionHistory)
+		SLATE_EVENT(SEditableSessionHistory::FCanDeleteActivity, CanDeleteActivity)
+		SLATE_EVENT(SEditableSessionHistory::FRequestDeleteActivity, DeleteActivity)
 		SLATE_NAMED_SLOT(FArguments, StatusBar)
 	SLATE_END_ARGS()
 
-	void Construct(const FArguments& InArgs, const FRequiredArgs& InRequiredArgs);
+	void Construct(const FArguments& InArgs);
 
 private:
 	
 	/** Holds the tab manager that manages the front-end's tabs. */
 	TSharedPtr<FTabManager> TabManager;
+
+	TSharedPtr<SEditableSessionHistory> SessionHistory;
 	
-	TSharedRef<SWidget> CreateTabs(const FRequiredArgs& RequiredArgs);
-	TSharedRef<SDockTab> SpawnActivityHistory(const FSpawnTabArgs& Args, TSharedRef<SSessionHistory> SessionHistory);
+	TSharedRef<SWidget> CreateTabs(const FArguments& InArgs);
+	TSharedRef<SDockTab> SpawnActivityHistory(const FSpawnTabArgs& Args,
+		SEditableSessionHistory::FMakeSessionHistory FMakeSessionHistory,
+		SEditableSessionHistory::FCanDeleteActivity CanDeleteActivity,
+		SEditableSessionHistory::FRequestDeleteActivity DeleteActivity
+		);
 };
diff --git a/Engine/Plugins/Developer/Concert/ConcertApp/MultiUserServer/Source/MultiUserServer/Private/Widgets/SessionTabs/Live/PackageViewer/SConcertSessionPackageViewer.cpp b/Engine/Plugins/Developer/Concert/ConcertApp/MultiUserServer/Source/MultiUserServer/Private/Widgets/SessionTabs/Live/PackageViewer/SConcertSessionPackageViewer.cpp
--- a/Engine/Plugins/Developer/Concert/ConcertApp/MultiUserServer/Source/MultiUserServer/Private/Widgets/SessionTabs/Live/PackageViewer/SConcertSessionPackageViewer.cpp
+++ b/Engine/Plugins/Developer/Concert/ConcertApp/MultiUserServer/Source/MultiUserServer/Private/Widgets/SessionTabs/Live/PackageViewer/SConcertSessionPackageViewer.cpp
@@ -20,7 +20,6 @@ void SConcertSessionPackageViewer::Construct(const FArguments& InArgs)
 
 	SAssignNew(ActivityListView, SConcertSessionActivities)
 		.OnGetPackageEvent(InArgs._GetPackageEvent)
-		.OnFetchActivities(InArgs._FetchInitialActivities)
 		.OnMapActivityToClient(InArgs._GetClientInfo)
 		.HighlightText(this, &SConcertSessionPackageViewer::HighlightSearchedText)
 		.Columns({
diff --git a/Engine/Plugins/Developer/Concert/ConcertApp/MultiUserServer/Source/MultiUserServer/Private/Widgets/SessionTabs/Live/PackageViewer/SConcertSessionPackageViewer.h b/Engine/Plugins/Developer/Concert/ConcertApp/MultiUserServer/Source/MultiUserServer/Private/Widgets/SessionTabs/Live/PackageViewer/SConcertSessionPackageViewer.h
--- a/Engine/Plugins/Developer/Concert/ConcertApp/MultiUserServer/Source/MultiUserServer/Private/Widgets/SessionTabs/Live/PackageViewer/SConcertSessionPackageViewer.h
+++ b/Engine/Plugins/Developer/Concert/ConcertApp/MultiUserServer/Source/MultiUserServer/Private/Widgets/SessionTabs/Live/PackageViewer/SConcertSessionPackageViewer.h
@@ -24,7 +24,6 @@ public:
 	SLATE_BEGIN_ARGS(SConcertSessionPackageViewer) {}
 		SLATE_EVENT(SConcertSessionActivities::FGetPackageEvent, GetPackageEvent)
 		SLATE_EVENT(SConcertSessionActivities::FGetActivityClientInfoFunc, GetClientInfo)
-		SLATE_EVENT(SConcertSessionActivities::FFetchActivitiesFunc, FetchInitialActivities)
 	SLATE_END_ARGS()
 
 	void Construct(const FArguments& InArgs);
diff --git a/Engine/Plugins/Developer/Concert/ConcertApp/MultiUserServer/Source/MultiUserServer/Private/Widgets/SessionTabs/ServerSessionHistoryControllerBase.cpp b/Engine/Plugins/Developer/Concert/ConcertApp/MultiUserServer/Source/MultiUserServer/Private/Widgets/SessionTabs/ServerSessionHistoryControllerBase.cpp
--- a/Engine/Plugins/Developer/Concert/ConcertApp/MultiUserServer/Source/MultiUserServer/Private/Widgets/SessionTabs/ServerSessionHistoryControllerBase.cpp
+++ b/Engine/Plugins/Developer/Concert/ConcertApp/MultiUserServer/Source/MultiUserServer/Private/Widgets/SessionTabs/ServerSessionHistoryControllerBase.cpp
@@ -6,8 +6,9 @@
 #include "IConcertSession.h"
 #include "IConcertSyncServer.h"
 
-FServerSessionHistoryControllerBase::FServerSessionHistoryControllerBase(FGuid SessionId)
-	: SessionId(MoveTemp(SessionId))
+FServerSessionHistoryControllerBase::FServerSessionHistoryControllerBase(FGuid SessionId, SSessionHistory::FArguments Arguments)
+	: FAbstractSessionHistoryController(MoveTemp(Arguments))
+	, SessionId(MoveTemp(SessionId))
 {}
 
 void FServerSessionHistoryControllerBase::GetActivities(int64 MaximumNumberOfActivities, TMap<FGuid, FConcertClientInfo>& OutEndpointClientInfoMap, TArray<FConcertSessionActivity>& OutFetchedActivities) const
diff --git a/Engine/Plugins/Developer/Concert/ConcertApp/MultiUserServer/Source/MultiUserServer/Private/Widgets/SessionTabs/ServerSessionHistoryControllerBase.h b/Engine/Plugins/Developer/Concert/ConcertApp/MultiUserServer/Source/MultiUserServer/Private/Widgets/SessionTabs/ServerSessionHistoryControllerBase.h
--- a/Engine/Plugins/Developer/Concert/ConcertApp/MultiUserServer/Source/MultiUserServer/Private/Widgets/SessionTabs/ServerSessionHistoryControllerBase.h
+++ b/Engine/Plugins/Developer/Concert/ConcertApp/MultiUserServer/Source/MultiUserServer/Private/Widgets/SessionTabs/ServerSessionHistoryControllerBase.h
@@ -19,7 +19,7 @@ class FServerSessionHistoryControllerBase : public FAbstractSessionHistoryContro
 {
 public:
 
-	FServerSessionHistoryControllerBase(FGuid SessionId);
+	FServerSessionHistoryControllerBase(FGuid SessionId, SSessionHistory::FArguments Arguments = SSessionHistory::FArguments());
 	
 protected:
 
diff --git a/Engine/Plugins/Developer/Concert/ConcertUI/ConcertSharedSlate/Source/ConcertSharedSlate/Private/Session/History/AbstractSessionHistoryController.cpp b/Engine/Plugins/Developer/Concert/ConcertUI/ConcertSharedSlate/Source/ConcertSharedSlate/Private/Session/History/AbstractSessionHistoryController.cpp
--- a/Engine/Plugins/Developer/Concert/ConcertUI/ConcertSharedSlate/Source/ConcertSharedSlate/Private/Session/History/AbstractSessionHistoryController.cpp
+++ b/Engine/Plugins/Developer/Concert/ConcertUI/ConcertSharedSlate/Source/ConcertSharedSlate/Private/Session/History/AbstractSessionHistoryController.cpp
@@ -5,8 +5,8 @@
 #include "Session/History/SSessionHistory.h"
 #include "Widgets/DeclarativeSyntaxSupport.h"
 
-FAbstractSessionHistoryController::FAbstractSessionHistoryController(FName PackageFilter)
-	: SessionHistory(MakeSessionHistory(PackageFilter))
+FAbstractSessionHistoryController::FAbstractSessionHistoryController(SSessionHistory::FArguments Arguments)
+	: SessionHistory(MakeSessionHistory(MoveTemp(Arguments)))
 {}
 
 void FAbstractSessionHistoryController::ReloadActivities()
@@ -23,10 +23,11 @@ void FAbstractSessionHistoryController::ReloadActivities()
 		);
 }
 
-TSharedRef<SSessionHistory> FAbstractSessionHistoryController::MakeSessionHistory(FName PackageFilter) const
+TSharedRef<SSessionHistory> FAbstractSessionHistoryController::MakeSessionHistory(SSessionHistory::FArguments Arguments) const
 {
-	return SNew(SSessionHistory)
-		.PackageFilter(PackageFilter)
-		.GetPackageEvent_Raw(this, &FAbstractSessionHistoryController::GetPackageEvent)
-		.GetTransactionEvent_Raw(this, &FAbstractSessionHistoryController::GetTransactionEvent);
+	return SArgumentNew(
+		Arguments
+			.GetPackageEvent_Raw(this, &FAbstractSessionHistoryController::GetPackageEvent)
+			.GetTransactionEvent_Raw(this, &FAbstractSessionHistoryController::GetTransactionEvent),
+		SSessionHistory);
 }
diff --git a/Engine/Plugins/Developer/Concert/ConcertUI/ConcertSharedSlate/Source/ConcertSharedSlate/Private/Session/History/SEditableSessionHistory.cpp b/Engine/Plugins/Developer/Concert/ConcertUI/ConcertSharedSlate/Source/ConcertSharedSlate/Private/Session/History/SEditableSessionHistory.cpp
new file mode 100644
--- /dev/null
+++ b/Engine/Plugins/Developer/Concert/ConcertUI/ConcertSharedSlate/Source/ConcertSharedSlate/Private/Session/History/SEditableSessionHistory.cpp
@@ -0,0 +1,77 @@
+﻿// Copyright Epic Games, Inc. All Rights Reserved.
+
+#include "Session/History/SEditableSessionHistory.h"
+
+#include "EditorStyleSet.h"
+#include "SNegativeActionButton.h"
+#include "Session/History/SSessionHistory.h"
+#include "Widgets/Layout/SBox.h"
+
+#define LOCTEXT_NAMESPACE "SEditableSessionHistory"
+
+void SEditableSessionHistory::Construct(const FArguments& InArgs)
+{
+	check(InArgs._MakeSessionHistory.IsBound() && InArgs._CanDeleteActivity.IsBound());
+	CanDeleteActivityFunc = InArgs._CanDeleteActivity;
+	DeleteActivityFunc = InArgs._DeleteActivity;
+	
+	SessionHistory = InArgs._MakeSessionHistory.Execute(
+			SSessionHistory::FArguments()
+			.OnMakeColumnOverlayWidget(this, &SEditableSessionHistory::MakeSummaryColumnDeleteButton)
+			);
+	
+	ChildSlot
+	[
+		SessionHistory.ToSharedRef()
+	];
+}
+
+TSharedPtr<SWidget> SEditableSessionHistory::MakeSummaryColumnDeleteButton(TWeakPtr<FConcertSessionActivity> RowActivity, const FName& ColumnId) const
+{
+	if (!SessionHistory->IsLastColumn(ColumnId))
+	{
+		return nullptr;
+	}
+	
+	return SNew(SBox)
+		.Padding(FMargin(1, 1))
+		.HAlign(HAlign_Right)
+		.VAlign(VAlign_Center)
+		[
+			SNew(SNegativeActionButton)
+			.OnClicked(this, &SEditableSessionHistory::OnClickDeleteActivityButton, RowActivity)
+			.ToolTipText(TAttribute<FText>::CreateLambda([this, RowActivity](){ return GetDeleteActivityToolTip(RowActivity);}))
+			.IsEnabled(TAttribute<bool>::CreateLambda([this, RowActivity](){ return CanDeleteActivity(RowActivity);}))
+			.Icon(FEditorStyle::GetBrush("Icons.Delete"))
+		];
+}
+
+FReply SEditableSessionHistory::OnClickDeleteActivityButton(TWeakPtr<FConcertSessionActivity> RowActivity) const
+{
+	if (DeleteActivityFunc.IsBound() && CanDeleteActivity(RowActivity))
+	{
+		DeleteActivityFunc.Execute(RowActivity.Pin().ToSharedRef());
+	}
+	return FReply::Handled();
+}
+
+FText SEditableSessionHistory::GetDeleteActivityToolTip(TWeakPtr<FConcertSessionActivity> RowActivity) const
+{
+	return CanDeleteActivity(RowActivity) ?
+		LOCTEXT("DeleteActivityDescription", "Delete activity from history")
+		:
+		LOCTEXT("CannotDeleteActivityDescription", "Activity cannot be deleted from history");
+}
+
+bool SEditableSessionHistory::CanDeleteActivity(TWeakPtr<FConcertSessionActivity> RowActivity) const
+{
+	if (const TSharedPtr<FConcertSessionActivity> Pinned = RowActivity.Pin())
+	{
+		return CanDeleteActivityFunc.Execute(Pinned.ToSharedRef());
+	}
+
+	checkNoEntry();
+	return false;
+}
+
+#undef LOCTEXT_NAMESPACE
diff --git a/Engine/Plugins/Developer/Concert/ConcertUI/ConcertSharedSlate/Source/ConcertSharedSlate/Private/Session/History/SSessionHistory.cpp b/Engine/Plugins/Developer/Concert/ConcertUI/ConcertSharedSlate/Source/ConcertSharedSlate/Private/Session/History/SSessionHistory.cpp
--- a/Engine/Plugins/Developer/Concert/ConcertUI/ConcertSharedSlate/Source/ConcertSharedSlate/Private/Session/History/SSessionHistory.cpp
+++ b/Engine/Plugins/Developer/Concert/ConcertUI/ConcertSharedSlate/Source/ConcertSharedSlate/Private/Session/History/SSessionHistory.cpp
@@ -53,6 +53,7 @@ void SSessionHistory::Construct(const FArguments& InArgs)
 	SAssignNew(ActivityListView, SConcertSessionActivities)
 		.OnGetPackageEvent(InArgs._GetPackageEvent)
 		.OnGetTransactionEvent(InArgs._GetTransactionEvent)
+		.OnMakeColumnOverlayWidget(InArgs._OnMakeColumnOverlayWidget)
 		.OnMapActivityToClient(this, &SSessionHistory::GetClientInfo)
 		.HighlightText(this, &SSessionHistory::HighlightSearchedText)
 		.TimeFormat(ActivityListViewOptions.Get(), &FConcertSessionActivitiesOptions::GetTimeFormat)
diff --git a/Engine/Plugins/Developer/Concert/ConcertUI/ConcertSharedSlate/Source/ConcertSharedSlate/Public/Session/History/AbstractSessionHistoryController.h b/Engine/Plugins/Developer/Concert/ConcertUI/ConcertSharedSlate/Source/ConcertSharedSlate/Public/Session/History/AbstractSessionHistoryController.h
--- a/Engine/Plugins/Developer/Concert/ConcertUI/ConcertSharedSlate/Source/ConcertSharedSlate/Public/Session/History/AbstractSessionHistoryController.h
+++ b/Engine/Plugins/Developer/Concert/ConcertUI/ConcertSharedSlate/Source/ConcertSharedSlate/Public/Session/History/AbstractSessionHistoryController.h
@@ -3,6 +3,8 @@
 #pragma once
 
 #include "CoreMinimal.h"
+#include "SSessionHistory.h"
+
 #include "Async/Future.h"
 
 struct FConcertClientInfo;
@@ -10,14 +12,13 @@ struct FConcertSyncTransactionEvent;
 struct FConcertSyncPackageEventMetaData;
 struct FConcertSessionActivity;
 class IConcertSyncClient;
-class SSessionHistory;
 
 /** Acts as controller of SSessionHistory in model-view-component pattern. */
 class CONCERTSHAREDSLATE_API FAbstractSessionHistoryController : public TSharedFromThis<FAbstractSessionHistoryController>
 {
 public:
 
-	FAbstractSessionHistoryController(FName PackageFilter = EName::None);
+	FAbstractSessionHistoryController(SSessionHistory::FArguments Arguments = SSessionHistory::FArguments());
 	virtual ~FAbstractSessionHistoryController() = default;
 	
 	/** Fetches activities from the server and updates the list view. */
@@ -41,5 +42,5 @@ private:
 	/** The widget being managed */
 	TSharedRef<SSessionHistory> SessionHistory;
 
-	TSharedRef<SSessionHistory> MakeSessionHistory(FName PackageFilter) const;
+	TSharedRef<SSessionHistory> MakeSessionHistory(SSessionHistory::FArguments Arguments) const;
 };
\ No newline at end of file
diff --git a/Engine/Plugins/Developer/Concert/ConcertUI/ConcertSharedSlate/Source/ConcertSharedSlate/Public/Session/History/SEditableSessionHistory.h b/Engine/Plugins/Developer/Concert/ConcertUI/ConcertSharedSlate/Source/ConcertSharedSlate/Public/Session/History/SEditableSessionHistory.h
new file mode 100644
--- /dev/null
+++ b/Engine/Plugins/Developer/Concert/ConcertUI/ConcertSharedSlate/Source/ConcertSharedSlate/Public/Session/History/SEditableSessionHistory.h
@@ -0,0 +1,44 @@
+﻿// Copyright Epic Games, Inc. All Rights Reserved.
+
+#pragma once
+
+#include "CoreMinimal.h"
+#include "Session/History/SSessionHistory.h"
+#include "Widgets/DeclarativeSyntaxSupport.h"
+#include "Widgets/SCompoundWidget.h"
+
+class SSessionHistory;
+struct FConcertSessionActivity;
+
+/** Allows activities in the session history to be deleted. */
+class CONCERTSHAREDSLATE_API SEditableSessionHistory : public SCompoundWidget
+{
+public:
+
+	DECLARE_DELEGATE_RetVal_OneParam(TSharedRef<SSessionHistory>, FMakeSessionHistory, SSessionHistory::FArguments)
+	DECLARE_DELEGATE_RetVal_OneParam(bool, FCanDeleteActivity, const TSharedRef<FConcertSessionActivity>&)
+	DECLARE_DELEGATE_OneParam(FRequestDeleteActivity, const TSharedRef<FConcertSessionActivity>&)
+
+	SLATE_BEGIN_ARGS(SEditableSessionHistory)
+	{}
+		SLATE_EVENT(FMakeSessionHistory, MakeSessionHistory)
+		SLATE_EVENT(FCanDeleteActivity, CanDeleteActivity)
+		SLATE_EVENT(FRequestDeleteActivity, DeleteActivity)
+		SLATE_NAMED_SLOT(FArguments, StatusBar)
+	SLATE_END_ARGS()
+
+	void Construct(const FArguments& InArgs);
+
+private:
+
+	TSharedPtr<SSessionHistory> SessionHistory;
+
+	FCanDeleteActivity CanDeleteActivityFunc;
+	FRequestDeleteActivity DeleteActivityFunc;
+
+	/** Creates a delete button that is overlayed over the summary column*/
+	TSharedPtr<SWidget> MakeSummaryColumnDeleteButton(TWeakPtr<FConcertSessionActivity> RowActivity, const FName& ColumnId) const;
+	FReply OnClickDeleteActivityButton(TWeakPtr<FConcertSessionActivity> RowActivity) const;
+	FText GetDeleteActivityToolTip(TWeakPtr<FConcertSessionActivity> RowActivity) const;
+	bool CanDeleteActivity(TWeakPtr<FConcertSessionActivity> RowActivity) const;
+};
diff --git a/Engine/Plugins/Developer/Concert/ConcertUI/ConcertSharedSlate/Source/ConcertSharedSlate/Public/Session/History/SSessionHistory.h b/Engine/Plugins/Developer/Concert/ConcertUI/ConcertSharedSlate/Source/ConcertSharedSlate/Public/Session/History/SSessionHistory.h
--- a/Engine/Plugins/Developer/Concert/ConcertUI/ConcertSharedSlate/Source/ConcertSharedSlate/Public/Session/History/SSessionHistory.h
+++ b/Engine/Plugins/Developer/Concert/ConcertUI/ConcertSharedSlate/Source/ConcertSharedSlate/Public/Session/History/SSessionHistory.h
@@ -8,7 +8,6 @@
 #include "Widgets/SCompoundWidget.h"
 
 class SSearchBox;
-class SConcertSessionActivities;
 class FConcertSessionActivitiesOptions;
 
 class CONCERTSHAREDSLATE_API SSessionHistory : public SCompoundWidget
@@ -22,6 +21,8 @@ public:
 		SLATE_ARGUMENT(FName, PackageFilter)
 		SLATE_EVENT(SConcertSessionActivities::FGetPackageEvent, GetPackageEvent)
 		SLATE_EVENT(SConcertSessionActivities::FGetTransactionEvent, GetTransactionEvent)
+		/** If bound, invoked when generating a row to add an overlay to a column. */
+		SLATE_EVENT(SConcertSessionActivities::FMakeColumnOverlayWidgetFunc, OnMakeColumnOverlayWidget)
 	SLATE_END_ARGS()
 
 	/**
@@ -39,6 +40,8 @@ public:
 	/** Callback for handling the a new or updated activity item. */ 
 	void HandleActivityAddedOrUpdated(const FConcertClientInfo& InClientInfo, const FConcertSyncActivity& InActivity, const FStructOnScope& InActivitySummary);
 
+	bool IsLastColumn(FName ColumnId) const { return ActivityListView->IsLastColumn(ColumnId); }
+
 private:
 
 	/** Invoked when the text in the search box widget changes. */
