Fix: Restored archived sessions do not respect initial rename

diff --git a/Engine/Plugins/Developer/Concert/ConcertApp/MultiUserServer/Source/MultiUserServer/Private/Widgets/Browser/ConcertServerSessionBrowserController.cpp b/Engine/Plugins/Developer/Concert/ConcertApp/MultiUserServer/Source/MultiUserServer/Private/Widgets/Browser/ConcertServerSessionBrowserController.cpp
--- a/Engine/Plugins/Developer/Concert/ConcertApp/MultiUserServer/Source/MultiUserServer/Private/Widgets/Browser/ConcertServerSessionBrowserController.cpp
+++ b/Engine/Plugins/Developer/Concert/ConcertApp/MultiUserServer/Source/MultiUserServer/Private/Widgets/Browser/ConcertServerSessionBrowserController.cpp
@@ -133,9 +133,10 @@ void FConcertServerSessionBrowserController::ArchiveSession(const FGuid& ServerA
 
 void FConcertServerSessionBrowserController::RestoreSession(const FGuid& ServerAdminEndpointId, const FGuid& SessionId, const FString& RestoredName, const FConcertSessionFilter& SessionFilter)
 {
-	if (const TOptional<FConcertSessionInfo> SessionInfo = GetArchivedSessionInfo(ServerAdminEndpointId, SessionId))
+	if (TOptional<FConcertSessionInfo> SessionInfo = GetArchivedSessionInfo(ServerAdminEndpointId, SessionId))
 	{
 		FText FailureReason = FText::GetEmpty();
+		SessionInfo->SessionName = RestoredName;
 		const bool bSuccess = ServerInstance->GetConcertServer()->RestoreSession(SessionId, *SessionInfo, SessionFilter, FailureReason).IsValid();
 		NotifyUserOfFinishedSessionAction(bSuccess,
 			bSuccess ? FText::Format(LOCTEXT("RestoreSessionFmt", "Restored Session '{0}'"),  FText::FromString(RestoredName)) : FText::Format(LOCTEXT("FailedToRestoreSessionFmt", "Failed to restore Session '{0}'"), FText::FromString(RestoredName)),
