Create system for fake modal dialogs

diff --git a/Engine/Plugins/Developer/Concert/ConcertApp/MultiUserServer/Source/MultiUserServer/Private/ConcertConsoleCommandExecutor.cpp b/Engine/Plugins/Developer/Concert/ConcertApp/MultiUserServer/Source/MultiUserServer/Private/ConcertConsoleCommandExecutor.cpp
--- a/Engine/Plugins/Developer/Concert/ConcertApp/MultiUserServer/Source/MultiUserServer/Private/ConcertConsoleCommandExecutor.cpp
+++ b/Engine/Plugins/Developer/Concert/ConcertApp/MultiUserServer/Source/MultiUserServer/Private/ConcertConsoleCommandExecutor.cpp
@@ -8,81 +8,84 @@
 
 #define LOCTEXT_NAMESPACE "UnrealMultiUserUI"
 
-FName FConcertConsoleCommandExecutor::StaticName()
+namespace UE::MultiUserServer
 {
-	static const FName CmdExecName = TEXT("Cmd");
-	return CmdExecName;
-}
+	FName FConcertConsoleCommandExecutor::StaticName()
+	{
+		static const FName CmdExecName = TEXT("Cmd");
+		return CmdExecName;
+	}
 
-FName FConcertConsoleCommandExecutor::GetName() const
-{
-	return StaticName();
-}
+	FName FConcertConsoleCommandExecutor::GetName() const
+	{
+		return StaticName();
+	}
 
-FText FConcertConsoleCommandExecutor::GetDisplayName() const
-{
-	return LOCTEXT("ConsoleCommandExecutorDisplayName", "Cmd");
-}
+	FText FConcertConsoleCommandExecutor::GetDisplayName() const
+	{
+		return LOCTEXT("ConsoleCommandExecutorDisplayName", "Cmd");
+	}
 
-FText FConcertConsoleCommandExecutor::GetDescription() const
-{
-	return LOCTEXT("ConsoleCommandExecutorDescription", "Execute Unreal Console Commands");
-}
+	FText FConcertConsoleCommandExecutor::GetDescription() const
+	{
+		return LOCTEXT("ConsoleCommandExecutorDescription", "Execute Unreal Console Commands");
+	}
 
-FText FConcertConsoleCommandExecutor::GetHintText() const
-{
-	return LOCTEXT("ConsoleCommandExecutorHintText", "Enter Console Command");
-}
+	FText FConcertConsoleCommandExecutor::GetHintText() const
+	{
+		return LOCTEXT("ConsoleCommandExecutorHintText", "Enter Console Command");
+	}
 
-void FConcertConsoleCommandExecutor::GetAutoCompleteSuggestions(const TCHAR* Input, TArray<FString>& Out)
-{
-	auto OnConsoleVariable = [&Out](const TCHAR *Name, IConsoleObject* CVar)
+	void FConcertConsoleCommandExecutor::GetAutoCompleteSuggestions(const TCHAR* Input, TArray<FString>& Out)
 	{
-		if (CVar->TestFlags(ECVF_Unregistered))
+		auto OnConsoleVariable = [&Out](const TCHAR *Name, IConsoleObject* CVar)
 		{
-			return;
-		}
+			if (CVar->TestFlags(ECVF_Unregistered))
+			{
+				return;
+			}
 
-		Out.Add(Name);
-	};
+			Out.Add(Name);
+		};
 
-	IConsoleManager::Get().ForEachConsoleObjectThatContains(FConsoleObjectVisitor::CreateLambda(OnConsoleVariable), Input);
-}
+		IConsoleManager::Get().ForEachConsoleObjectThatContains(FConsoleObjectVisitor::CreateLambda(OnConsoleVariable), Input);
+	}
 
-void FConcertConsoleCommandExecutor::GetExecHistory(TArray<FString>& Out)
-{
-	IConsoleManager::Get().GetConsoleHistory(TEXT(""), Out);
-}
+	void FConcertConsoleCommandExecutor::GetExecHistory(TArray<FString>& Out)
+	{
+		IConsoleManager::Get().GetConsoleHistory(TEXT(""), Out);
+	}
 
-bool FConcertConsoleCommandExecutor::Exec(const TCHAR* Input)
-{
-	IConsoleManager::Get().AddConsoleHistoryEntry(TEXT(""), Input);
+	bool FConcertConsoleCommandExecutor::Exec(const TCHAR* Input)
+	{
+		IConsoleManager::Get().AddConsoleHistoryEntry(TEXT(""), Input);
 	
-	FOutputDevice& Output = *GLog;
-	UWorld* World = nullptr;
+		FOutputDevice& Output = *GLog;
+		UWorld* World = nullptr;
 	
-	bool bSuccess = IConsoleManager::Get().ProcessUserConsoleInput(Input, Output, World);
-	if (!bSuccess)
-	{
-		bSuccess |= FSelfRegisteringExec::StaticExec(World, Input, Output);
-	}
+		bool bSuccess = IConsoleManager::Get().ProcessUserConsoleInput(Input, Output, World);
+		if (!bSuccess)
+		{
+			bSuccess |= FSelfRegisteringExec::StaticExec(World, Input, Output);
+		}
 
-	return bSuccess;
-}
+		return bSuccess;
+	}
 
-bool FConcertConsoleCommandExecutor::AllowHotKeyClose() const
-{
-	return true;
-}
+	bool FConcertConsoleCommandExecutor::AllowHotKeyClose() const
+	{
+		return true;
+	}
 
-bool FConcertConsoleCommandExecutor::AllowMultiLine() const
-{
-	return false;
-}
+	bool FConcertConsoleCommandExecutor::AllowMultiLine() const
+	{
+		return false;
+	}
 
-FInputChord FConcertConsoleCommandExecutor::GetHotKey() const
-{
-	return FInputChord();
+	FInputChord FConcertConsoleCommandExecutor::GetHotKey() const
+	{
+		return FInputChord();
+	}
 }
 
 #undef LOCTEXT_NAMESPACE
diff --git a/Engine/Plugins/Developer/Concert/ConcertApp/MultiUserServer/Source/MultiUserServer/Private/ConcertConsoleCommandExecutor.h b/Engine/Plugins/Developer/Concert/ConcertApp/MultiUserServer/Source/MultiUserServer/Private/ConcertConsoleCommandExecutor.h
--- a/Engine/Plugins/Developer/Concert/ConcertApp/MultiUserServer/Source/MultiUserServer/Private/ConcertConsoleCommandExecutor.h
+++ b/Engine/Plugins/Developer/Concert/ConcertApp/MultiUserServer/Source/MultiUserServer/Private/ConcertConsoleCommandExecutor.h
@@ -5,23 +5,26 @@
 #include "CoreMinimal.h"
 #include "HAL/IConsoleManager.h"
 
-/** Executes commands put into the output log */
-class FConcertConsoleCommandExecutor : public IConsoleCommandExecutor
+namespace UE::MultiUserServer
 {
-public:
+	/** Executes commands put into the output log */
+	class FConcertConsoleCommandExecutor : public IConsoleCommandExecutor
+	{
+	public:
 	
-	static FName StaticName();
+		static FName StaticName();
 
-	//~ Begin IConsoleCommandExecutor Interface
-	virtual FName GetName() const override;
-	virtual FText GetDisplayName() const override;
-	virtual FText GetDescription() const override;
-	virtual FText GetHintText() const override;
-	virtual void GetAutoCompleteSuggestions(const TCHAR* Input, TArray<FString>& Out) override;
-	virtual void GetExecHistory(TArray<FString>& Out) override;
-	virtual bool Exec(const TCHAR* Input) override;
-	virtual bool AllowHotKeyClose() const override;
-	virtual bool AllowMultiLine() const override;
-	virtual FInputChord GetHotKey() const override;
-	//~ End IConsoleCommandExecutor Interface
-};
\ No newline at end of file
+		//~ Begin IConsoleCommandExecutor Interface
+		virtual FName GetName() const override;
+		virtual FText GetDisplayName() const override;
+		virtual FText GetDescription() const override;
+		virtual FText GetHintText() const override;
+		virtual void GetAutoCompleteSuggestions(const TCHAR* Input, TArray<FString>& Out) override;
+		virtual void GetExecHistory(TArray<FString>& Out) override;
+		virtual bool Exec(const TCHAR* Input) override;
+		virtual bool AllowHotKeyClose() const override;
+		virtual bool AllowMultiLine() const override;
+		virtual FInputChord GetHotKey() const override;
+		//~ End IConsoleCommandExecutor Interface
+	};
+}
diff --git a/Engine/Plugins/Developer/Concert/ConcertApp/MultiUserServer/Source/MultiUserServer/Private/ModalWindowManager.cpp b/Engine/Plugins/Developer/Concert/ConcertApp/MultiUserServer/Source/MultiUserServer/Private/ModalWindowManager.cpp
new file mode 100644
--- /dev/null
+++ b/Engine/Plugins/Developer/Concert/ConcertApp/MultiUserServer/Source/MultiUserServer/Private/ModalWindowManager.cpp
@@ -0,0 +1,69 @@
+﻿// Copyright Epic Games, Inc. All Rights Reserved.
+
+#include "ModalWindowManager.h"
+
+#include "Dialog/SCustomDialog.h"
+#include "Widgets/SWindow.h"
+
+UE::MultiUserServer::FModalWindowManager::FModalWindowManager(const TSharedRef<SWindow>& MainWindow)
+{
+	MainWindow->GetOnWindowActivatedEvent().AddRaw(this, &FModalWindowManager::OnMainWindowActivated);
+
+	// Fake modal windows should be closed when the root window is closed
+	MainWindow->GetOnWindowClosedEvent().AddLambda([this](const TSharedRef<SWindow>&)
+	{
+		// One way to close root window while fake modal is open:
+			// 1. Hover MU server window icon (on Windows)
+			// 2. Close root window
+		bRootWindowIsValid = false;
+		
+		for (TWeakPtr<SWindow> ModalWindow : TArray(WindowStack))
+		{
+			if (TSharedPtr<SWindow> Pin = ModalWindow.Pin())
+			{
+				Pin->RequestDestroyWindow();
+			}
+		}
+		WindowStack.Empty();
+	});
+}
+
+void UE::MultiUserServer::FModalWindowManager::ShowFakeModalWindow(const TSharedRef<SCustomDialog>& Dialog)
+{
+	if (ensureMsgf(bRootWindowIsValid, TEXT("Root window has already been destroyed!")))
+	{
+		Dialog->Show();
+		PushWindow(Dialog);
+	}
+}
+
+void UE::MultiUserServer::FModalWindowManager::PushWindow(const TSharedRef<SWindow>& Window)
+{
+	if (ensureMsgf(bRootWindowIsValid, TEXT("Root window has already been destroyed!")))
+	{
+		check(!WindowStack.Contains(Window));
+		WindowStack.Push(Window);
+	
+		Window->GetOnWindowClosedEvent().AddLambda([this](const TSharedRef<SWindow>&)
+		{
+			PopWindow();
+		});
+	}
+}
+
+void UE::MultiUserServer::FModalWindowManager::PopWindow()
+{
+	WindowStack.Pop();
+}
+
+void UE::MultiUserServer::FModalWindowManager::OnMainWindowActivated()
+{
+	if (WindowStack.Num() > 0)
+	{
+		if (const TSharedPtr<SWindow> PinnedWindow = WindowStack.Top().Pin())
+		{
+			PinnedWindow->BringToFront();
+			PinnedWindow->FlashWindow();
+		}
+	}
+}
diff --git a/Engine/Plugins/Developer/Concert/ConcertApp/MultiUserServer/Source/MultiUserServer/Private/ModalWindowManager.h b/Engine/Plugins/Developer/Concert/ConcertApp/MultiUserServer/Source/MultiUserServer/Private/ModalWindowManager.h
new file mode 100644
--- /dev/null
+++ b/Engine/Plugins/Developer/Concert/ConcertApp/MultiUserServer/Source/MultiUserServer/Private/ModalWindowManager.h
@@ -0,0 +1,42 @@
+﻿// Copyright Epic Games, Inc. All Rights Reserved.
+
+#pragma once
+
+#include "CoreMinimal.h"
+
+class SWindow;
+class SCustomDialog;
+
+namespace UE::MultiUserServer
+{
+	/**
+	 * Manages "fake" modal windows because the server cannot use "true" modal windows.
+	 *
+	 * True modal windows block the main thread which is not an option because the server must continue responding to
+	 * networked requests.
+	 */
+	class FModalWindowManager
+	{
+	public:
+
+		FModalWindowManager(const TSharedRef<SWindow>& MainWindow);
+
+		void ShowFakeModalWindow(const TSharedRef<SCustomDialog>& Dialog);
+
+		/** Called when a new fake modal window should be tracked on top of all windows. */
+		void PushWindow(const TSharedRef<SWindow>& Window);
+		/** Called when a fake modal window should no longer be tracked, i.e. when it is closed. */
+		void PopWindow();
+
+	private:
+
+		TArray<TWeakPtr<SWindow>> WindowStack;
+
+		/** Modal windows can only be added when this is true */
+		bool bRootWindowIsValid = true;
+
+		void OnMainWindowActivated();
+	};
+}
+
+
diff --git a/Engine/Plugins/Developer/Concert/ConcertApp/MultiUserServer/Source/MultiUserServer/Private/MultiUserServerModule.cpp b/Engine/Plugins/Developer/Concert/ConcertApp/MultiUserServer/Source/MultiUserServer/Private/MultiUserServerModule.cpp
--- a/Engine/Plugins/Developer/Concert/ConcertApp/MultiUserServer/Source/MultiUserServer/Private/MultiUserServerModule.cpp
+++ b/Engine/Plugins/Developer/Concert/ConcertApp/MultiUserServer/Source/MultiUserServer/Private/MultiUserServerModule.cpp
@@ -1,11 +1,13 @@
 // Copyright Epic Games, Inc. All Rights Reserved.
 
-#include "IMultiUserServerModule.h"
+#include "MultiUserServerModule.h"
 
 #include "ConcertConsoleCommandExecutor.h"
 #include "ConcertFrontendStyle.h"
 #include "ConcertServerStyle.h"
 #include "ConcertSyncServerLoopInitArgs.h"
+#include "ModalWindowManager.h"
+
 #include "Widgets/ConcertServerWindowController.h"
 
 #include "Features/IModularFeatures.h"
@@ -15,26 +17,23 @@
 
 #define LOCTEXT_NAMESPACE "UnrealMultiUserUI"
 
-/**
- * 
- */
-class FConcertServerUIModule : public IMultiUserServerModule
+namespace UE::MultiUserServer
 {
-public:
-	virtual void StartupModule() override 
+    void FConcertServerUIModule::StartupModule() 
 	{
 		MultiUserServerLayoutIni = GConfig->GetConfigFilename(TEXT("MultiUserServerLayout"));
 	}
 	
-	virtual void ShutdownModule() override
+	void FConcertServerUIModule::ShutdownModule()
 	{
 		FConcertServerStyle::Shutdown();
-		
+
+    	ModalWindowManager.Reset();
 		WindowController.Reset();
 		FSlateApplication::Shutdown();
 	}
 	
-	virtual void InitSlateForServer(FConcertSyncServerLoopInitArgs& InitArgs) override
+	void FConcertServerUIModule::InitSlateForServer(FConcertSyncServerLoopInitArgs& InitArgs)
 	{
 		if (ensureMsgf(!WindowController.IsValid(), TEXT("InitSlateForServer is designed to be called at most once to create UI to run alongside the Multi User server.")))
 		{
@@ -44,47 +43,38 @@ public:
 		}
 	}
 	
-private:
+	void FConcertServerUIModule::PreInitializeMultiUser()
+    {
+    	FModuleManager::Get().LoadModuleChecked("EditorStyle");
+    	FConcertServerStyle::Initialize();
+    	FConcertFrontendStyle::Initialize();
 
-	/** Config path storing layout config. */
-	FString MultiUserServerLayoutIni;
+    	// Log history must be initialized before MU server loop init prints any messages	
+    	FModuleManager::Get().LoadModuleChecked("OutputLog");
+    }
+    
+    void FConcertServerUIModule::InitializeSlateApplication(TSharedRef<IConcertSyncServer> SyncServer)
+    {
+    	FSlateApplication::InitializeAsStandaloneApplication(GetStandardStandaloneRenderer());
+    	const FText ApplicationTitle = LOCTEXT("AppTitle", "Unreal Multi User Server");
+    	FGlobalTabmanager::Get()->SetApplicationTitle(ApplicationTitle);
 
-	/** Handles execution of commands */
-	TUniquePtr<FConcertConsoleCommandExecutor> CommandExecutor;
-	
-	/** Creates and manages window. Only one instance per application. */
-	TSharedPtr<FConcertServerWindowController> WindowController;
+    	CommandExecutor = MakeUnique<FConcertConsoleCommandExecutor>();
+    	IModularFeatures::Get().RegisterModularFeature(IConsoleCommandExecutor::ModularFeatureName(), CommandExecutor.Get());
 
-	void PreInitializeMultiUser()
-	{
-		FModuleManager::Get().LoadModuleChecked("EditorStyle");
-		FConcertServerStyle::Initialize();
-		FConcertFrontendStyle::Initialize();
-
-		// Log history must be initialized before MU server loop init prints any messages	
-		FModuleManager::Get().LoadModuleChecked("OutputLog");
-	}
-	
-	void InitializeSlateApplication(TSharedRef<IConcertSyncServer> SyncServer)
-	{
-		FSlateApplication::InitializeAsStandaloneApplication(GetStandardStandaloneRenderer());
-		const FText ApplicationTitle = LOCTEXT("AppTitle", "Unreal Multi User Server");
-		FGlobalTabmanager::Get()->SetApplicationTitle(ApplicationTitle);
-
-		CommandExecutor = MakeUnique<FConcertConsoleCommandExecutor>();
-		IModularFeatures::Get().RegisterModularFeature(IConsoleCommandExecutor::ModularFeatureName(), CommandExecutor.Get());
-        
         WindowController = MakeShared<FConcertServerWindowController>(FConcertServerWindowInitParams{ SyncServer, MultiUserServerLayoutIni });
-		WindowController->CreateWindow();
-	}
-	
-	void TickSlate(double Tick)
-	{
-		FSlateApplication::Get().PumpMessages();
-		FSlateApplication::Get().Tick();
-	}
-};
+    	const TSharedRef<SWindow> MainWindow = WindowController->CreateWindow();
+    	
+    	ModalWindowManager = MakeShared<FModalWindowManager>(MainWindow);
+    }
+    
+    void FConcertServerUIModule::TickSlate(double Tick)
+    {
+    	FSlateApplication::Get().PumpMessages();
+    	FSlateApplication::Get().Tick();
+    }
+}
 
-IMPLEMENT_MODULE(FConcertServerUIModule, ConcertSyncServerUI);
+IMPLEMENT_MODULE(UE::MultiUserServer::FConcertServerUIModule, MultiUserServer);
 
 #undef LOCTEXT_NAMESPACE
\ No newline at end of file
diff --git a/Engine/Plugins/Developer/Concert/ConcertApp/MultiUserServer/Source/MultiUserServer/Private/MultiUserServerModule.h b/Engine/Plugins/Developer/Concert/ConcertApp/MultiUserServer/Source/MultiUserServer/Private/MultiUserServerModule.h
new file mode 100644
--- /dev/null
+++ b/Engine/Plugins/Developer/Concert/ConcertApp/MultiUserServer/Source/MultiUserServer/Private/MultiUserServerModule.h
@@ -0,0 +1,59 @@
+﻿// Copyright Epic Games, Inc. All Rights Reserved.
+
+#pragma once
+
+#include "CoreMinimal.h"
+
+#include "ConcertSyncServerLoopInitArgs.h"
+#include "IMultiUserServerModule.h"
+
+#include "Modules/ModuleManager.h"
+
+class FConcertServerWindowController;
+
+namespace UE::MultiUserServer
+{
+	class FModalWindowManager;
+	class FConcertConsoleCommandExecutor;
+	
+	class FConcertServerUIModule : public IMultiUserServerModule
+    {
+    public:
+    
+		static inline FConcertServerUIModule& Get()
+    	{
+    		static const FName ModuleName = "MultiUserServer";
+    		return FModuleManager::LoadModuleChecked<FConcertServerUIModule>(ModuleName);
+    	}
+
+		TSharedPtr<FModalWindowManager> GetModalWindowManager() const { return ModalWindowManager; }
+            
+        //~ Begin IModuleInterface Interface
+        virtual void StartupModule() override ;
+        virtual void ShutdownModule() override;
+        //~ End IModuleInterface Interface
+        
+        //~ Begin IMultiUserServerModule Interface
+        virtual void InitSlateForServer(FConcertSyncServerLoopInitArgs& InitArgs) override;
+        //~ End IMultiUserServerModule Interface
+        
+    private:
+    
+        /** Config path storing layout config. */
+        FString MultiUserServerLayoutIni;
+    
+        /** Handles execution of commands */
+        TUniquePtr<FConcertConsoleCommandExecutor> CommandExecutor;
+        
+		/** Manages fake modal windows */
+		TSharedPtr<FModalWindowManager> ModalWindowManager;
+		
+        /** Creates and manages window. Only one instance per application. */
+        TSharedPtr<FConcertServerWindowController> WindowController;
+    
+        void PreInitializeMultiUser();
+        void InitializeSlateApplication(TSharedRef<IConcertSyncServer> SyncServer);
+        
+        void TickSlate(double Tick);
+    };
+}
diff --git a/Engine/Plugins/Developer/Concert/ConcertApp/MultiUserServer/Source/MultiUserServer/Private/Widgets/Browser/SConcertServerSessionBrowser.cpp b/Engine/Plugins/Developer/Concert/ConcertApp/MultiUserServer/Source/MultiUserServer/Private/Widgets/Browser/SConcertServerSessionBrowser.cpp
--- a/Engine/Plugins/Developer/Concert/ConcertApp/MultiUserServer/Source/MultiUserServer/Private/Widgets/Browser/SConcertServerSessionBrowser.cpp
+++ b/Engine/Plugins/Developer/Concert/ConcertApp/MultiUserServer/Source/MultiUserServer/Private/Widgets/Browser/SConcertServerSessionBrowser.cpp
@@ -2,8 +2,9 @@
 
 #include "SConcertServerSessionBrowser.h"
 
+#include "ModalWindowManager.h"
+#include "MultiUserServerModule.h"
 #include "MultiUserServerUserSettings.h"
-
 #include "Session/Browser/SConcertSessionBrowser.h"
 #include "Widgets/Browser/ConcertServerSessionBrowserController.h"
 #include "Widgets/StatusBar/SConcertStatusBar.h"
@@ -18,11 +19,6 @@
 
 #define LOCTEXT_NAMESPACE "UnrealMultiUserUI"
 
-SConcertServerSessionBrowser::~SConcertServerSessionBrowser()
-{
-	UnregisterFromOnRootWindowClosed();
-}
-
 void SConcertServerSessionBrowser::Construct(const FArguments& InArgs, TSharedRef<FConcertServerSessionBrowserController> InController)
 {
 	Controller = InController;
@@ -72,53 +68,18 @@ TSharedRef<SWidget> SConcertServerSessionBrowser::MakeSessionTableView(const FAr
 
 void SConcertServerSessionBrowser::RequestDeleteSession(const TSharedPtr<FConcertSessionItem>& SessionItem) 
 {
-	if (const TSharedPtr<SMessageDialog> PinnedDeleteSessionDialog = DeleteSessionDialog.Pin())
-	{
-		// Force user to answer previous question first
-		PinnedDeleteSessionDialog->BringToFront();
-		PinnedDeleteSessionDialog->FlashWindow();
-	}
-	else
-	{
-		// Modal dialog would cause server to halt thread and stop communicating with connected clients
-		switch (SessionItem->Type)
-		{
-		case FConcertSessionItem::EType::ActiveSession:
-			DeleteActiveSessionWithNonModalQuestion(SessionItem); break;
-		case FConcertSessionItem::EType::ArchivedSession:
-			DeleteArchivedSessionWithNonModalQuestion(SessionItem); break;
-		default:
-			break;
-		}
-
-		// The dialog should be closed when the root window is closed
-		if (const TSharedPtr<SWindow> RootWindow = FGlobalTabmanager::Get()->GetRootWindow();
-			RootWindow && !RootWindow->GetOnWindowClosedEvent().IsBoundToObject(this))
-		{
-			RootWindow->GetOnWindowClosedEvent().AddSP(this, &SConcertServerSessionBrowser::OnRootWindowClosed);
-		}
-	}
-}
-
-void SConcertServerSessionBrowser::OnRootWindowClosed(const TSharedRef<SWindow>&) const
-{
-	if (const TSharedPtr<SMessageDialog> PinnedDeleteSessionDialog = DeleteSessionDialog.Pin())
-	{
-		PinnedDeleteSessionDialog->RequestDestroyWindow();
-	}
-	
-	UnregisterFromOnRootWindowClosed();
-}
-
-void SConcertServerSessionBrowser::UnregisterFromOnRootWindowClosed() const
-{
-	if (const TSharedPtr<SWindow> RootWindow = FGlobalTabmanager::Get()->GetRootWindow())
+	switch (SessionItem->Type)
 	{
-		RootWindow->GetOnWindowClosedEvent().RemoveAll(this);
+	case FConcertSessionItem::EType::ActiveSession:
+		DeleteActiveSessionWithFakeModalQuestion(SessionItem); break;
+	case FConcertSessionItem::EType::ArchivedSession:
+		DeleteArchivedSessionWithFakeModalQuestion(SessionItem); break;
+	default:
+		break;
 	}
 }
 
-void SConcertServerSessionBrowser::DeleteArchivedSessionWithNonModalQuestion(const TSharedPtr<FConcertSessionItem>& SessionItem)
+void SConcertServerSessionBrowser::DeleteArchivedSessionWithFakeModalQuestion(const TSharedPtr<FConcertSessionItem>& SessionItem)
 {
 	const FText Message = FText::Format(
 	LOCTEXT("DeleteArchivedDescription", "Deleting a session will cause all associated data to be removed.\n\nDelete {0}?"),
@@ -138,17 +99,21 @@ void SConcertServerSessionBrowser::DeleteArchivedSessionWithNonModalQuestion(con
 		.Title(LOCTEXT("DisconnectUsersTitle", "Delete session?"))
 		.IconBrush("Icons.WarningWithColor.Large")
 		.Message(Message)
+		.UseScrollBox(false)
 		.Buttons({
 			SMessageDialog::FButton(LOCTEXT("DeleteArchivedButton", "Delete"))
 				.SetOnClicked(FSimpleDelegate::CreateLambda(DeleteArchived)),
-			SMessageDialog::FButton(LOCTEXT("CancelButton", "Cancel")).SetPrimary(true)
+			SMessageDialog::FButton(LOCTEXT("CancelButton", "Cancel"))
+				.SetPrimary(true)
+				.SetFocus()
 		});
-	
-	DeleteSessionDialog = Dialog;
-	Dialog->Show();
+
+	UE::MultiUserServer::FConcertServerUIModule::Get()
+		.GetModalWindowManager()
+		->ShowFakeModalWindow(Dialog);
 }
 
-void SConcertServerSessionBrowser::DeleteActiveSessionWithNonModalQuestion(const TSharedPtr<FConcertSessionItem>& SessionItem)
+void SConcertServerSessionBrowser::DeleteActiveSessionWithFakeModalQuestion(const TSharedPtr<FConcertSessionItem>& SessionItem)
 {
 	const int32 NumUsers = Controller.Pin()->GetNumConnectedClients(SessionItem->SessionId);
 	const FText Message = FText::Format(
@@ -170,14 +135,18 @@ void SConcertServerSessionBrowser::DeleteActiveSessionWithNonModalQuestion(const
 		.Title(LOCTEXT("DisconnectUsersTitle", "Force Users to Disconnect?"))
 		.IconBrush("Icons.WarningWithColor.Large")
 		.Message(Message)
+		.UseScrollBox(false)
 		.Buttons({
-			SMessageDialog::FButton(LOCTEXT("DeleteActiveButton", "Delete")).SetPrimary(true)
-				.SetOnClicked(FSimpleDelegate::CreateLambda(DeleteActive)),
+			SMessageDialog::FButton(LOCTEXT("DeleteActiveButton", "Delete"))
+				.SetPrimary(true)
+				.SetOnClicked(FSimpleDelegate::CreateLambda(DeleteActive))
+				.SetFocus(),
 			SMessageDialog::FButton(LOCTEXT("CancelButton", "Cancel"))
 		});
-	
-	DeleteSessionDialog = Dialog;
-	Dialog->Show();
+
+	UE::MultiUserServer::FConcertServerUIModule::Get()
+		.GetModalWindowManager()
+		->ShowFakeModalWindow(Dialog);
 }
 
 #undef LOCTEXT_NAMESPACE
diff --git a/Engine/Plugins/Developer/Concert/ConcertApp/MultiUserServer/Source/MultiUserServer/Private/Widgets/Browser/SConcertServerSessionBrowser.h b/Engine/Plugins/Developer/Concert/ConcertApp/MultiUserServer/Source/MultiUserServer/Private/Widgets/Browser/SConcertServerSessionBrowser.h
--- a/Engine/Plugins/Developer/Concert/ConcertApp/MultiUserServer/Source/MultiUserServer/Private/Widgets/Browser/SConcertServerSessionBrowser.h
+++ b/Engine/Plugins/Developer/Concert/ConcertApp/MultiUserServer/Source/MultiUserServer/Private/Widgets/Browser/SConcertServerSessionBrowser.h
@@ -20,7 +20,6 @@ public:
 		SLATE_EVENT(FSessionDelegate, DoubleClickLiveSession)
 		SLATE_EVENT(FSessionDelegate, DoubleClickArchivedSession)
 	SLATE_END_ARGS()
-	virtual ~SConcertServerSessionBrowser() override;
 
 	void Construct(const FArguments& InArgs, TSharedRef<FConcertServerSessionBrowserController> InController);
 
@@ -40,8 +39,6 @@ private:
 
 	/** We can ask the controller about information and notify it about UI events. */
 	TWeakPtr<FConcertServerSessionBrowserController> Controller;
-	/** Tracks whether there is a dialog asking the user to delete a session. Used to avoid opening multiple. */
-	TWeakPtr<SMessageDialog> DeleteSessionDialog;
 
 	bool bRequestedRefresh = false;
 
@@ -51,8 +48,6 @@ private:
 	TSharedRef<SWidget> MakeSessionTableView(const FArguments& InArgs);
 
 	void RequestDeleteSession(const TSharedPtr<FConcertSessionItem>& SessionItem);
-	void OnRootWindowClosed(const TSharedRef<SWindow>&) const;
-	void UnregisterFromOnRootWindowClosed() const;
-	void DeleteArchivedSessionWithNonModalQuestion(const TSharedPtr<FConcertSessionItem>& SessionItem);
-	void DeleteActiveSessionWithNonModalQuestion(const TSharedPtr<FConcertSessionItem>& SessionItem);
+	void DeleteArchivedSessionWithFakeModalQuestion(const TSharedPtr<FConcertSessionItem>& SessionItem);
+	void DeleteActiveSessionWithFakeModalQuestion(const TSharedPtr<FConcertSessionItem>& SessionItem);
 };
diff --git a/Engine/Plugins/Developer/Concert/ConcertApp/MultiUserServer/Source/MultiUserServer/Private/Widgets/ConcertServerWindowController.cpp b/Engine/Plugins/Developer/Concert/ConcertApp/MultiUserServer/Source/MultiUserServer/Private/Widgets/ConcertServerWindowController.cpp
--- a/Engine/Plugins/Developer/Concert/ConcertApp/MultiUserServer/Source/MultiUserServer/Private/Widgets/ConcertServerWindowController.cpp
+++ b/Engine/Plugins/Developer/Concert/ConcertApp/MultiUserServer/Source/MultiUserServer/Private/Widgets/ConcertServerWindowController.cpp
@@ -29,7 +29,7 @@ FConcertServerWindowController::FConcertServerWindowController(const FConcertSer
 	SessionBrowserController = MakeShared<FConcertServerSessionBrowserController>();
 }
 
-void FConcertServerWindowController::CreateWindow()
+TSharedRef<SWindow> FConcertServerWindowController::CreateWindow()
 {
 	InitComponents();
 
@@ -78,6 +78,8 @@ void FConcertServerWindowController::CreateWindow()
 	RootWindow->ShowWindow();
 	constexpr bool bForceWindowToFront = true;
 	RootWindow->BringToFront(bForceWindowToFront);
+
+	return RootWindowRef;
 }
 
 void FConcertServerWindowController::OpenSessionTab(const FGuid& SessionId)
diff --git a/Engine/Plugins/Developer/Concert/ConcertApp/MultiUserServer/Source/MultiUserServer/Private/Widgets/ConcertServerWindowController.h b/Engine/Plugins/Developer/Concert/ConcertApp/MultiUserServer/Source/MultiUserServer/Private/Widgets/ConcertServerWindowController.h
--- a/Engine/Plugins/Developer/Concert/ConcertApp/MultiUserServer/Source/MultiUserServer/Private/Widgets/ConcertServerWindowController.h
+++ b/Engine/Plugins/Developer/Concert/ConcertApp/MultiUserServer/Source/MultiUserServer/Private/Widgets/ConcertServerWindowController.h
@@ -33,7 +33,7 @@ class FConcertServerWindowController : public TSharedFromThis<FConcertServerWind
 public:
 	
 	FConcertServerWindowController(const FConcertServerWindowInitParams& Params);
-	void CreateWindow();
+	TSharedRef<SWindow> CreateWindow();
 
 	/** Opens or draws attention to the tab for the given live or archived session ID */
 	void OpenSessionTab(const FGuid& SessionId);
