Create icons for session content

diff --git a/Engine/Plugins/Developer/Concert/ConcertApp/MultiUserServer/Source/MultiUserServer/Private/ConcertServerStyle.cpp b/Engine/Plugins/Developer/Concert/ConcertApp/MultiUserServer/Source/MultiUserServer/Private/ConcertServerStyle.cpp
--- a/Engine/Plugins/Developer/Concert/ConcertApp/MultiUserServer/Source/MultiUserServer/Private/ConcertServerStyle.cpp
+++ b/Engine/Plugins/Developer/Concert/ConcertApp/MultiUserServer/Source/MultiUserServer/Private/ConcertServerStyle.cpp
@@ -63,6 +63,11 @@ TSharedRef<FSlateStyleSet> FConcertServerStyle::Create()
 	const FVector2D Icon32x32(20.0f, 20.0f); 
 	
 	StyleSet->Set("Concert.MultiUser", new IMAGE_PLUGIN_BRUSH("Icons/icon_MultiUser_32x", Icon32x32));
+	StyleSet->Set("Concert.SessionContent.ColumnHeader", new IMAGE_PLUGIN_BRUSH_SVG("Icons/Package_16x", Icon16x16));
+	StyleSet->Set("Concert.SessionContent.PackageAdded", new IMAGE_PLUGIN_BRUSH_SVG("Icons/PackageAdded_16x", Icon16x16));
+	StyleSet->Set("Concert.SessionContent.PackageDeleted", new IMAGE_PLUGIN_BRUSH_SVG("Icons/PackageDeleted_16x", Icon16x16));
+	StyleSet->Set("Concert.SessionContent.PackageRenamed", new IMAGE_PLUGIN_BRUSH_SVG("Icons/PackageRenamed_16x", Icon16x16));
+	StyleSet->Set("Concert.SessionContent.PackageSaved", new IMAGE_PLUGIN_BRUSH_SVG("Icons/PackageSaved_16x", Icon16x16));
 	return StyleSet;
 }
 
diff --git a/Engine/Plugins/Developer/Concert/ConcertApp/MultiUserServer/Source/MultiUserServer/Private/Widgets/SessionTabs/Live/PackageViewer/PackageViewerColumns.cpp b/Engine/Plugins/Developer/Concert/ConcertApp/MultiUserServer/Source/MultiUserServer/Private/Widgets/SessionTabs/Live/PackageViewer/PackageViewerColumns.cpp
--- a/Engine/Plugins/Developer/Concert/ConcertApp/MultiUserServer/Source/MultiUserServer/Private/Widgets/SessionTabs/Live/PackageViewer/PackageViewerColumns.cpp
+++ b/Engine/Plugins/Developer/Concert/ConcertApp/MultiUserServer/Source/MultiUserServer/Private/Widgets/SessionTabs/Live/PackageViewer/PackageViewerColumns.cpp
@@ -2,13 +2,18 @@
 
 #include "PackageViewerColumns.h"
 
+#include "ConcertServerStyle.h"
+
 #include "Math/UnitConversion.h"
+
+#include "Widgets/Layout/SScaleBox.h"
 #include "Widgets/Text/STextBlock.h"
 
 #define LOCTEXT_NAMESPACE "UnrealMultiUserUI"
 
-const FName UE::MultiUserServer::PackageViewerColumns::SizeColumnId		= TEXT("SizeColumnId");
-const FName UE::MultiUserServer::PackageViewerColumns::VersionColumnId	= TEXT("VersionColumnId");
+const FName UE::MultiUserServer::PackageViewerColumns::PackageUpdateTypeColumnId	= TEXT("PackageUpdateTypeColumnId");
+const FName UE::MultiUserServer::PackageViewerColumns::SizeColumnId					= TEXT("SizeColumnId");
+const FName UE::MultiUserServer::PackageViewerColumns::VersionColumnId				= TEXT("VersionColumnId");
 
 namespace UE::MultiUserServer::PackageViewerColumns
 {
@@ -54,6 +59,66 @@ namespace UE::MultiUserServer::PackageViewerColumns
 	}
 }
 
+FActivityColumn UE::MultiUserServer::PackageViewerColumns::PackageUpdateTypeColumn(FGetPackageUpdateType GetPackageUpdateTypeFunc)
+{
+	return FActivityColumn(
+		SHeaderRow::Column(PackageUpdateTypeColumnId)
+			.HeaderContent()
+			[
+				SNew(SScaleBox)
+				.Stretch(EStretch::ScaleToFit)
+				[
+					SNew(SImage)
+					.Image(FConcertServerStyle::Get().GetBrush("Concert.SessionContent.ColumnHeader"))
+				]
+			]
+			.FixedWidth(22)
+			// Cannot be hidden
+			.ShouldGenerateWidget(true)
+		)
+		.ColumnSortOrder(static_cast<int32>(EPredefinedPackageColumnOrder::PackageUpdateType))
+		.GenerateColumnWidget(FActivityColumn::FGenerateColumnWidget::CreateLambda([GetPackageUpdateTypeFunc](const TSharedRef<SConcertSessionActivities>& Owner, const TSharedRef<FConcertSessionActivity>& Activity, SOverlay::FScopedWidgetSlotArguments& Slot)
+		{
+			if (const TOptional<EConcertPackageUpdateType> PackageUpdateType = GetPackageUpdateTypeFunc.Execute(Activity.Get()))
+			{
+				const TMap<EConcertPackageUpdateType, FName> PackageIcons {
+					{ EConcertPackageUpdateType::Added, TEXT("Concert.SessionContent.PackageAdded") },
+					{ EConcertPackageUpdateType::Deleted, TEXT("Concert.SessionContent.PackageDeleted") },
+					{ EConcertPackageUpdateType::Renamed, TEXT("Concert.SessionContent.PackageRenamed") },
+					{ EConcertPackageUpdateType::Saved, TEXT("Concert.SessionContent.PackageSaved") },
+					
+					// This one is not expected to be found but it's here for safety
+					{ EConcertPackageUpdateType::Dummy, TEXT("Concert.SessionContent.ColumnHeader")}
+				};
+				const TMap<EConcertPackageUpdateType, FText> ToolTips {
+					{ EConcertPackageUpdateType::Added, LOCTEXT("PackageUpdateTypeColumn.Tooltip.Added", "Added packaged") },
+					{ EConcertPackageUpdateType::Deleted, LOCTEXT("PackageUpdateTypeColumn.Tooltip.Deleted", "Deleted package") },
+					{ EConcertPackageUpdateType::Renamed, LOCTEXT("PackageUpdateTypeColumn.Tooltip.Renamed", "Renamed package") },
+					{ EConcertPackageUpdateType::Saved, LOCTEXT("PackageUpdateTypeColumn.Tooltip.Saved", "Saved package") },
+					
+					// This one is not expected to be found but it's here for safety
+					{ EConcertPackageUpdateType::Dummy, LOCTEXT("PackageUpdateTypeColumn.Tooltip.Dummy", "Dummy") }
+				};
+				static_assert(static_cast<int32>(EConcertPackageUpdateType::Count) == 5, "You must update this TMap when you update EConcertPackageUpdateType.");
+				
+				Slot
+				.HAlign(HAlign_Center)
+				.Padding(1)
+				[
+					SNew(SImage)
+					.Image(FConcertServerStyle::Get().GetBrush(PackageIcons[*PackageUpdateType]))
+					.DesiredSizeOverride(FVector2D(16.f, 16.f))
+					.ToolTipText(ToolTips[*PackageUpdateType])
+				];
+			}
+			else
+			{
+				checkNoEntry();
+				Slot[SNullWidget::NullWidget];
+			}
+		}));
+}
+
 FActivityColumn UE::MultiUserServer::PackageViewerColumns::SizeColumn(FGetSizeOfPackageActivity GetEventDataFunc)
 {
 	return CreateNumericColumn(
diff --git a/Engine/Plugins/Developer/Concert/ConcertApp/MultiUserServer/Source/MultiUserServer/Private/Widgets/SessionTabs/Live/PackageViewer/PackageViewerColumns.h b/Engine/Plugins/Developer/Concert/ConcertApp/MultiUserServer/Source/MultiUserServer/Private/Widgets/SessionTabs/Live/PackageViewer/PackageViewerColumns.h
--- a/Engine/Plugins/Developer/Concert/ConcertApp/MultiUserServer/Source/MultiUserServer/Private/Widgets/SessionTabs/Live/PackageViewer/PackageViewerColumns.h
+++ b/Engine/Plugins/Developer/Concert/ConcertApp/MultiUserServer/Source/MultiUserServer/Private/Widgets/SessionTabs/Live/PackageViewer/PackageViewerColumns.h
@@ -8,19 +8,23 @@
 
 namespace UE::MultiUserServer::PackageViewerColumns
 {
+	DECLARE_DELEGATE_RetVal_OneParam(TOptional<EConcertPackageUpdateType>, FGetPackageUpdateType, const FConcertSessionActivity& /*Activity*/);
 	DECLARE_DELEGATE_RetVal_OneParam(TOptional<int64>, FGetNumericValueFromPackageActivity, const FConcertSessionActivity& /*Activity*/);
 	using FGetSizeOfPackageActivity = FGetNumericValueFromPackageActivity;
 	using FGetVersionOfPackageActivity = FGetNumericValueFromPackageActivity;
 	
 	enum class EPredefinedPackageColumnOrder : int32
 	{
+		PackageUpdateType = static_cast<int32>(ConcertSharedSlate::ActivityColumn::EPredefinedColumnOrder::AvatarColor) - 1,
 		Size = static_cast<int32>(ConcertSharedSlate::ActivityColumn::EPredefinedColumnOrder::ClientName) + 1,
 		Version = Size + 1
 	};
-	
+
+	extern const FName PackageUpdateTypeColumnId;
 	extern const FName SizeColumnId;
 	extern const FName VersionColumnId;
-	
+
+	FActivityColumn PackageUpdateTypeColumn(FGetPackageUpdateType GetPackageUpdateTypeFunc);
 	FActivityColumn SizeColumn(FGetSizeOfPackageActivity GetEventDataFunc);
 	FActivityColumn VersionColumn(FGetVersionOfPackageActivity GetEventDataFunc);
 };
diff --git a/Engine/Plugins/Developer/Concert/ConcertApp/MultiUserServer/Source/MultiUserServer/Private/Widgets/SessionTabs/Live/PackageViewer/SConcertSessionPackageViewer.cpp b/Engine/Plugins/Developer/Concert/ConcertApp/MultiUserServer/Source/MultiUserServer/Private/Widgets/SessionTabs/Live/PackageViewer/SConcertSessionPackageViewer.cpp
--- a/Engine/Plugins/Developer/Concert/ConcertApp/MultiUserServer/Source/MultiUserServer/Private/Widgets/SessionTabs/Live/PackageViewer/SConcertSessionPackageViewer.cpp
+++ b/Engine/Plugins/Developer/Concert/ConcertApp/MultiUserServer/Source/MultiUserServer/Private/Widgets/SessionTabs/Live/PackageViewer/SConcertSessionPackageViewer.cpp
@@ -29,6 +29,7 @@ void SConcertSessionPackageViewer::Construct(const FArguments& InArgs)
 		.OnMapActivityToClient(InArgs._GetClientInfo)
 		.HighlightText(this, &SConcertSessionPackageViewer::HighlightSearchedText)
 		.Columns({
+			PackageViewerColumns::PackageUpdateTypeColumn(PackageViewerColumns::FGetPackageUpdateType::CreateSP(this, &SConcertSessionPackageViewer::GetPackageActivityUpdateType, InArgs._GetPackageEvent)),
 			PackageViewerColumns::SizeColumn(InArgs._GetSizeOfPackageActivity),
 			PackageViewerColumns::VersionColumn(PackageViewerColumns::FGetVersionOfPackageActivity::CreateSP(this, &SConcertSessionPackageViewer::GetVersionOfPackageActivity, InArgs._GetPackageEvent))
 		})
@@ -95,6 +96,16 @@ void SConcertSessionPackageViewer::OnColumnVisibilitySettingsChanged(const FColu
 	ActivityListView->OnColumnVisibilitySettingsChanged(ColumnSnapshot);
 }
 
+TOptional<EConcertPackageUpdateType> SConcertSessionPackageViewer::GetPackageActivityUpdateType(const FConcertSessionActivity& Activity, SConcertSessionActivities::FGetPackageEvent GetPackageEventFunc) const
+{
+	FConcertSyncPackageEventMetaData PackageEventMetaData;
+	if (GetPackageEventFunc.Execute(Activity, PackageEventMetaData))
+	{
+		return PackageEventMetaData.PackageInfo.PackageUpdateType;
+	}
+	return {};
+}
+
 TOptional<int64> SConcertSessionPackageViewer::GetVersionOfPackageActivity(const FConcertSessionActivity& Activity, SConcertSessionActivities::FGetPackageEvent GetPackageEventFunc) const
 {
 	FConcertSyncPackageEventMetaData PackageEventMetaData;
diff --git a/Engine/Plugins/Developer/Concert/ConcertApp/MultiUserServer/Source/MultiUserServer/Private/Widgets/SessionTabs/Live/PackageViewer/SConcertSessionPackageViewer.h b/Engine/Plugins/Developer/Concert/ConcertApp/MultiUserServer/Source/MultiUserServer/Private/Widgets/SessionTabs/Live/PackageViewer/SConcertSessionPackageViewer.h
--- a/Engine/Plugins/Developer/Concert/ConcertApp/MultiUserServer/Source/MultiUserServer/Private/Widgets/SessionTabs/Live/PackageViewer/SConcertSessionPackageViewer.h
+++ b/Engine/Plugins/Developer/Concert/ConcertApp/MultiUserServer/Source/MultiUserServer/Private/Widgets/SessionTabs/Live/PackageViewer/SConcertSessionPackageViewer.h
@@ -43,6 +43,7 @@ private:
 	/** The searched text to highlight. */
 	FText SearchedText;
 
+	TOptional<EConcertPackageUpdateType> GetPackageActivityUpdateType(const FConcertSessionActivity& Activity, SConcertSessionActivities::FGetPackageEvent GetPackageEventFunc) const;
 	TOptional<int64> GetVersionOfPackageActivity(const FConcertSessionActivity& Activity, SConcertSessionActivities::FGetPackageEvent GetPackageEventFunc) const;
 
 	/** Invoked when the text in the search box widget changes. */
@@ -51,4 +52,4 @@ private:
 	void OnSearchTextCommitted(const FText& InFilterText, ETextCommit::Type CommitType);
 	/** Returns the text to highlight when the search bar has a text set. */
 	FText HighlightSearchedText() const;
-};
\ No newline at end of file
+};
diff --git a/Engine/Plugins/Developer/Concert/ConcertSync/ConcertSyncCore/Source/ConcertSyncCore/Public/ConcertWorkspaceData.h b/Engine/Plugins/Developer/Concert/ConcertSync/ConcertSyncCore/Source/ConcertSyncCore/Public/ConcertWorkspaceData.h
--- a/Engine/Plugins/Developer/Concert/ConcertSync/ConcertSyncCore/Source/ConcertSyncCore/Public/ConcertWorkspaceData.h
+++ b/Engine/Plugins/Developer/Concert/ConcertSync/ConcertSyncCore/Source/ConcertSyncCore/Public/ConcertWorkspaceData.h
@@ -19,6 +19,8 @@ enum class EConcertPackageUpdateType : uint8
 	Renamed,
 	/** This package has been deleted */
 	Deleted,
+
+	Count
 };
 
 USTRUCT()
diff --git a/Engine/Plugins/Developer/Concert/ConcertUI/ConcertSharedSlate/Source/ConcertSharedSlate/Private/Session/Activity/PredefinedActivityColumns.cpp b/Engine/Plugins/Developer/Concert/ConcertUI/ConcertSharedSlate/Source/ConcertSharedSlate/Private/Session/Activity/PredefinedActivityColumns.cpp
--- a/Engine/Plugins/Developer/Concert/ConcertUI/ConcertSharedSlate/Source/ConcertSharedSlate/Private/Session/Activity/PredefinedActivityColumns.cpp
+++ b/Engine/Plugins/Developer/Concert/ConcertUI/ConcertSharedSlate/Source/ConcertSharedSlate/Private/Session/Activity/PredefinedActivityColumns.cpp
@@ -103,7 +103,7 @@ FActivityColumn UE::ConcertSharedSlate::ActivityColumn::AvatarColor()
 	return FActivityColumn(
 		SHeaderRow::Column(AvatarColorColumnId)
 			.DefaultLabel(INVTEXT(""))
-			.ManualWidth(8)
+			.FixedWidth(8)
 			.ShouldGenerateWidget(true)
 		)
 		.ColumnSortOrder(static_cast<int32>(EPredefinedColumnOrder::AvatarColor))
