Push missing file

diff --git a/Engine/Plugins/Developer/Concert/ConcertApp/MultiUserServer/Source/MultiUserServer/Private/Widgets/SessionTabs/Live/PackageViewer/ConcertSessionPackageViewerController.cpp b/Engine/Plugins/Developer/Concert/ConcertApp/MultiUserServer/Source/MultiUserServer/Private/Widgets/SessionTabs/Live/PackageViewer/ConcertSessionPackageViewerController.cpp
--- a/Engine/Plugins/Developer/Concert/ConcertApp/MultiUserServer/Source/MultiUserServer/Private/Widgets/SessionTabs/Live/PackageViewer/ConcertSessionPackageViewerController.cpp
+++ b/Engine/Plugins/Developer/Concert/ConcertApp/MultiUserServer/Source/MultiUserServer/Private/Widgets/SessionTabs/Live/PackageViewer/ConcertSessionPackageViewerController.cpp
@@ -4,6 +4,7 @@
 
 #include "ConcertSyncSessionDatabase.h"
 #include "IConcertSyncServer.h"
+#include "MultiUserServerUserSettings.h"
 #include "SConcertSessionPackageViewer.h"
 #include "Concert/Private/ConcertServerSession.h"
 
@@ -13,6 +14,7 @@ FConcertSessionPackageViewerController::FConcertSessionPackageViewerController(T
 	, PackageViewer(MakePackageViewer())
 {
 	ReloadActivities();
+	UMultiUserServerUserSettings::GetUserSettings()->OnLiveSessionContentColumnVisibility().AddRaw(this, &FConcertSessionPackageViewerController::OnSessionContentColumnVisibilitySettingsUpdated);
 	
 	if (const TOptional<FConcertSyncSessionDatabaseNonNullPtr> Database = SyncServer->GetLiveSessionDatabase(InspectedSession->GetId()))
 	{
@@ -22,6 +24,11 @@ FConcertSessionPackageViewerController::FConcertSessionPackageViewerController(T
 
 FConcertSessionPackageViewerController::~FConcertSessionPackageViewerController()
 {
+	if (UMultiUserServerUserSettings* Settings = UMultiUserServerUserSettings::GetUserSettings(); IsValid(Settings))
+	{
+		Settings->OnLiveSessionContentColumnVisibility().RemoveAll(this);
+	}
+	
 	if (const TOptional<FConcertSyncSessionDatabaseNonNullPtr> Database = SyncServer->GetLiveSessionDatabase(InspectedSession->GetId()))
 	{
 		Database->OnActivityProduced().RemoveAll(this);
@@ -124,3 +131,8 @@ void FConcertSessionPackageViewerController::OnSessionProduced(const FConcertSyn
 		ReloadActivities();
 	}
 }
+
+void FConcertSessionPackageViewerController::OnSessionContentColumnVisibilitySettingsUpdated(const FColumnVisibilitySnapshot& NewValue)
+{
+	GetPackageViewer()->OnColumnVisibilitySettingsChanged(NewValue);
+}
