Change tab name when session is renamed

diff --git a/Engine/Plugins/Developer/Concert/ConcertApp/MultiUserServer/Source/MultiUserServer/Private/Widgets/SessionTabs/ConcertSessionTabBase.cpp b/Engine/Plugins/Developer/Concert/ConcertApp/MultiUserServer/Source/MultiUserServer/Private/Widgets/SessionTabs/ConcertSessionTabBase.cpp
--- a/Engine/Plugins/Developer/Concert/ConcertApp/MultiUserServer/Source/MultiUserServer/Private/Widgets/SessionTabs/ConcertSessionTabBase.cpp
+++ b/Engine/Plugins/Developer/Concert/ConcertApp/MultiUserServer/Source/MultiUserServer/Private/Widgets/SessionTabs/ConcertSessionTabBase.cpp
@@ -52,12 +52,14 @@ void FConcertSessionTabBase::EnsureInitDockTab()
 {
 	if (!DockTab.IsValid())
 	{
-		const TOptional<FString> SessionName = UE::ConcertServerUI::Private::GetSessionName(SyncServer, GetSessionID());
-		check(SessionName);
-		
-		const FText Title = FText::FromString(*SessionName);
 		DockTab = SNew(SDockTab)
-			.Label(Title)
+			.Label_Lambda([this]()
+			{
+				return FText::FromString(
+					UE::ConcertServerUI::Private::GetSessionName(SyncServer, GetSessionID())
+						.Get(TEXT("Error getting name"))
+					);
+			})
 			.TabRole(MajorTab);
 		
 		CreateDockContent(DockTab.ToSharedRef());
