Fix not being able to rename live sessions on server

diff --git a/Engine/Plugins/Developer/Concert/ConcertApp/MultiUserServer/Source/MultiUserServer/Private/Widgets/Browser/ConcertServerSessionBrowserController.cpp b/Engine/Plugins/Developer/Concert/ConcertApp/MultiUserServer/Source/MultiUserServer/Private/Widgets/Browser/ConcertServerSessionBrowserController.cpp
--- a/Engine/Plugins/Developer/Concert/ConcertApp/MultiUserServer/Source/MultiUserServer/Private/Widgets/Browser/ConcertServerSessionBrowserController.cpp
+++ b/Engine/Plugins/Developer/Concert/ConcertApp/MultiUserServer/Source/MultiUserServer/Private/Widgets/Browser/ConcertServerSessionBrowserController.cpp
@@ -175,20 +175,33 @@ void FConcertServerSessionBrowserController::OpenSession(const TSharedPtr<FConce
 	Owner.Pin()->OpenSessionTab(SessionItem->SessionId);
 }
 
-void FConcertServerSessionBrowserController::RenameSession(const FGuid& ServerAdminEndpointId, const FGuid& SessionId, const FString& NewName)
+void FConcertServerSessionBrowserController::RenameActiveSessionInternal(const FGuid& ServerAdminEndpointId, const FGuid& SessionId, const FString& NewName)
+{
+	if (const TOptional<FConcertSessionInfo> SessionInfo = GetActiveSessionInfo(ServerAdminEndpointId, SessionId))
+	{
+		RenameSessionInternal(SessionId, NewName, *SessionInfo);
+	}
+}
+
+void FConcertServerSessionBrowserController::RenameArchivedSessionInternal(const FGuid& ServerAdminEndpointId, const FGuid& SessionId, const FString& NewName)
 {
 	if (const TOptional<FConcertSessionInfo> SessionInfo = GetArchivedSessionInfo(ServerAdminEndpointId, SessionId))
 	{
-		FText FailureReason = FText::GetEmpty();
-		const bool bSuccess = ServerInstance->GetConcertServer()->RenameSession(SessionId, NewName, FailureReason);
-		NotifyUserOfFinishedSessionAction(bSuccess,
-				bSuccess ? FText::Format(LOCTEXT("RenameSessionFmt", "Rename Session '{0}' as '{1}'"), FText::FromString(SessionInfo->SessionName), FText::FromString(NewName))
-					: FText::Format(LOCTEXT("FailedToArchivedSessionFmt", "Failed to rename Session '{0}' as '{1}'"), FText::FromString(SessionInfo->SessionName), FText::FromString(NewName)),
-				FailureReason
-				);
+		RenameSessionInternal(SessionId, NewName, *SessionInfo);
 	}
 }
 
+void FConcertServerSessionBrowserController::RenameSessionInternal(const FGuid& SessionId, const FString& NewName, const FConcertSessionInfo& SessionInfo)
+{
+	FText FailureReason = FText::GetEmpty();
+	const bool bSuccess = ServerInstance->GetConcertServer()->RenameSession(SessionId, NewName, FailureReason);
+	NotifyUserOfFinishedSessionAction(bSuccess,
+			bSuccess ? FText::Format(LOCTEXT("RenameSessionFmt", "Rename Session '{0}' as '{1}'"), FText::FromString(SessionInfo.SessionName), FText::FromString(NewName))
+				: FText::Format(LOCTEXT("FailedToArchivedSessionFmt", "Failed to rename Session '{0}' as '{1}'"), FText::FromString(SessionInfo.SessionName), FText::FromString(NewName)),
+			FailureReason
+			);
+}
+
 void FConcertServerSessionBrowserController::DeleteSession(const FGuid& ServerAdminEndpointId, const FGuid& SessionId)
 {
 	const TOptional<FConcertSessionInfo> ArchivedInfo = GetArchivedSessionInfo(ServerAdminEndpointId, SessionId);
diff --git a/Engine/Plugins/Developer/Concert/ConcertApp/MultiUserServer/Source/MultiUserServer/Private/Widgets/Browser/ConcertServerSessionBrowserController.h b/Engine/Plugins/Developer/Concert/ConcertApp/MultiUserServer/Source/MultiUserServer/Private/Widgets/Browser/ConcertServerSessionBrowserController.h
--- a/Engine/Plugins/Developer/Concert/ConcertApp/MultiUserServer/Source/MultiUserServer/Private/Widgets/Browser/ConcertServerSessionBrowserController.h
+++ b/Engine/Plugins/Developer/Concert/ConcertApp/MultiUserServer/Source/MultiUserServer/Private/Widgets/Browser/ConcertServerSessionBrowserController.h
@@ -40,8 +40,8 @@ public:
 	virtual void CreateSession(const FGuid& ServerAdminEndpointId, const FString& SessionName, const FString& ProjectName) override;
 	virtual void ArchiveSession(const FGuid& ServerAdminEndpointId, const FGuid& SessionId, const FString& ArchiveName, const FConcertSessionFilter& SessionFilter) override;
 	virtual void RestoreSession(const FGuid& ServerAdminEndpointId, const FGuid& SessionId, const FString& RestoredName, const FConcertSessionFilter& SessionFilter) override;
-	virtual void RenameActiveSession(const FGuid& ServerAdminEndpointId, const FGuid& SessionId, const FString& NewName) override { RenameSession(ServerAdminEndpointId, SessionId, NewName); }
-	virtual void RenameArchivedSession(const FGuid& ServerAdminEndpointId, const FGuid& SessionId, const FString& NewName) override { RenameSession(ServerAdminEndpointId, SessionId, NewName); }
+	virtual void RenameActiveSession(const FGuid& ServerAdminEndpointId, const FGuid& SessionId, const FString& NewName) override { RenameActiveSessionInternal(ServerAdminEndpointId, SessionId, NewName); }
+	virtual void RenameArchivedSession(const FGuid& ServerAdminEndpointId, const FGuid& SessionId, const FString& NewName) override { RenameArchivedSessionInternal(ServerAdminEndpointId, SessionId, NewName); }
 	virtual void DeleteActiveSession(const FGuid& ServerAdminEndpointId, const FGuid& SessionId) override { DeleteSession(ServerAdminEndpointId, SessionId); }
 	virtual void DeleteArchivedSession(const FGuid& ServerAdminEndpointId, const FGuid& SessionId) override { DeleteSession(ServerAdminEndpointId, SessionId); }
 
@@ -73,7 +73,9 @@ private:
 
 	// Session actions
 	void OpenSession(const TSharedPtr<FConcertSessionItem>& SessionItem);
-	void RenameSession(const FGuid& ServerAdminEndpointId, const FGuid& SessionId, const FString& NewName);
+	void RenameActiveSessionInternal(const FGuid& ServerAdminEndpointId, const FGuid& SessionId, const FString& NewName);
+	void RenameArchivedSessionInternal(const FGuid& ServerAdminEndpointId, const FGuid& SessionId, const FString& NewName);
+	void RenameSessionInternal(const FGuid& SessionId, const FString& NewName, const FConcertSessionInfo& SessionInfo);
 	void DeleteSession(const FGuid& ServerAdminEndpointId, const FGuid& SessionId);
 
 	// Notifications
