Deleting a session should close the tab that displays the session

diff --git a/Engine/Plugins/Developer/Concert/ConcertApp/MultiUserServer/Source/MultiUserServer/Private/Widgets/ConcertServerWindowController.cpp b/Engine/Plugins/Developer/Concert/ConcertApp/MultiUserServer/Source/MultiUserServer/Private/Widgets/ConcertServerWindowController.cpp
--- a/Engine/Plugins/Developer/Concert/ConcertApp/MultiUserServer/Source/MultiUserServer/Private/Widgets/ConcertServerWindowController.cpp
+++ b/Engine/Plugins/Developer/Concert/ConcertApp/MultiUserServer/Source/MultiUserServer/Private/Widgets/ConcertServerWindowController.cpp
@@ -2,6 +2,8 @@
 
 #include "ConcertServerWindowController.h"
 
+#include "ConcertServerEvents.h"
+
 #include "Browser/ConcertServerSessionBrowserController.h"
 #include "ConcertServerTabs.h"
 #include "IConcertServer.h"
@@ -29,6 +31,11 @@ FConcertServerWindowController::FConcertServerWindowController(const FConcertSer
 	SessionBrowserController = MakeShared<FConcertServerSessionBrowserController>();
 }
 
+FConcertServerWindowController::~FConcertServerWindowController()
+{
+	UnregisterFromSessionDestructionEvents();
+}
+
 TSharedRef<SWindow> FConcertServerWindowController::CreateWindow()
 {
 	InitComponents();
@@ -79,6 +86,7 @@ TSharedRef<SWindow> FConcertServerWindowController::CreateWindow()
 	constexpr bool bForceWindowToFront = true;
 	RootWindow->BringToFront(bForceWindowToFront);
 
+	RegisterForSessionDestructionEvents();
 	return RootWindowRef;
 }
 
@@ -90,9 +98,15 @@ void FConcertServerWindowController::OpenSessionTab(const FGuid& SessionId)
 	}
 }
 
+void FConcertServerWindowController::DestroySessionTab(const FGuid& SessionId)
+{
+	// Destructor will handle the rest, e.g. removing the tab from the window
+	RegisteredSessions.Remove(SessionId);
+}
+
 TSharedPtr<FConcertSessionTabBase> FConcertServerWindowController::GetOrRegisterSessionTab(const FGuid& SessionId)
 {
-	if (TSharedRef<FConcertSessionTabBase>* FoundId = RegisteredSessions.Find(SessionId))
+	if (const TSharedRef<FConcertSessionTabBase>* FoundId = RegisteredSessions.Find(SessionId))
 	{
 		return *FoundId;
 	}
@@ -121,6 +135,28 @@ void FConcertServerWindowController::InitComponents()
 	SessionBrowserController->Init(Params);
 }
 
+void FConcertServerWindowController::RegisterForSessionDestructionEvents()
+{
+	ConcertServerEvents::OnLiveSessionDestroyed().AddSP(this, &FConcertServerWindowController::OnLiveSessionDestroyed);
+	ConcertServerEvents::OnArchivedSessionDestroyed().AddSP(this, &FConcertServerWindowController::OnArchivedSessionDestroyed);
+}
+
+void FConcertServerWindowController::UnregisterFromSessionDestructionEvents() const
+{
+	ConcertServerEvents::OnLiveSessionDestroyed().RemoveAll(this);
+	ConcertServerEvents::OnArchivedSessionDestroyed().RemoveAll(this);
+}
+
+void FConcertServerWindowController::OnLiveSessionDestroyed(const IConcertServer&, TSharedRef<IConcertServerSession> InLiveSession)
+{
+	DestroySessionTab(InLiveSession->GetId());
+}
+
+void FConcertServerWindowController::OnArchivedSessionDestroyed(const IConcertServer&, const FGuid& InArchivedSessionId)
+{
+	DestroySessionTab(InArchivedSessionId);
+}
+
 void FConcertServerWindowController::OnWindowClosed(const TSharedRef<SWindow>& Window)
 {
 	SaveLayout();
diff --git a/Engine/Plugins/Developer/Concert/ConcertApp/MultiUserServer/Source/MultiUserServer/Private/Widgets/ConcertServerWindowController.h b/Engine/Plugins/Developer/Concert/ConcertApp/MultiUserServer/Source/MultiUserServer/Private/Widgets/ConcertServerWindowController.h
--- a/Engine/Plugins/Developer/Concert/ConcertApp/MultiUserServer/Source/MultiUserServer/Private/Widgets/ConcertServerWindowController.h
+++ b/Engine/Plugins/Developer/Concert/ConcertApp/MultiUserServer/Source/MultiUserServer/Private/Widgets/ConcertServerWindowController.h
@@ -5,6 +5,7 @@
 #include "CoreMinimal.h"
 #include "Framework/Docking/TabManager.h"
 
+class IConcertServer;
 class FConcertSessionTabBase;
 class FLiveConcertSessionTab;
 class IConcertServerSession;
@@ -33,10 +34,14 @@ class FConcertServerWindowController : public TSharedFromThis<FConcertServerWind
 public:
 	
 	FConcertServerWindowController(const FConcertServerWindowInitParams& Params);
+	~FConcertServerWindowController();
 	TSharedRef<SWindow> CreateWindow();
 
 	/** Opens or draws attention to the tab for the given live or archived session ID */
 	void OpenSessionTab(const FGuid& SessionId);
+
+	/** Destroys the tab associated with this live or archived session ID. */
+	void DestroySessionTab(const FGuid& SessionId);
 	
 private:
 
@@ -59,6 +64,11 @@ private:
 	/** Gets the manager for a session tab if the session ID is valid */
 	TSharedPtr<FConcertSessionTabBase> GetOrRegisterSessionTab(const FGuid& SessionId);
 
+	void RegisterForSessionDestructionEvents();
+	void UnregisterFromSessionDestructionEvents() const;
+	void OnLiveSessionDestroyed(const IConcertServer&, TSharedRef<IConcertServerSession> InLiveSession);
+	void OnArchivedSessionDestroyed(const IConcertServer&, const FGuid& InArchivedSessionId);
+	
 	void OnWindowClosed(const TSharedRef<SWindow>& Window);
 	void SaveLayout() const;
 };
diff --git a/Engine/Plugins/Developer/Concert/ConcertApp/MultiUserServer/Source/MultiUserServer/Private/Widgets/SessionTabs/Archived/ArchivedConcertSessionTab.cpp b/Engine/Plugins/Developer/Concert/ConcertApp/MultiUserServer/Source/MultiUserServer/Private/Widgets/SessionTabs/Archived/ArchivedConcertSessionTab.cpp
--- a/Engine/Plugins/Developer/Concert/ConcertApp/MultiUserServer/Source/MultiUserServer/Private/Widgets/SessionTabs/Archived/ArchivedConcertSessionTab.cpp
+++ b/Engine/Plugins/Developer/Concert/ConcertApp/MultiUserServer/Source/MultiUserServer/Private/Widgets/SessionTabs/Archived/ArchivedConcertSessionTab.cpp
@@ -9,17 +9,12 @@
 #include "Widgets/StatusBar/SConcertStatusBar.h"
 
 FArchivedConcertSessionTab::FArchivedConcertSessionTab(FGuid InspectedSessionID, TSharedRef<IConcertSyncServer> SyncServer, TAttribute<TSharedRef<SWindow>> ConstructUnderWindow)
-	: FConcertSessionTabBase(SyncServer)
+	: FConcertSessionTabBase(InspectedSessionID, SyncServer)
 	, InspectedSessionID(MoveTemp(InspectedSessionID))
 	, SyncServer(MoveTemp(SyncServer))
 	, ConstructUnderWindow(MoveTemp(ConstructUnderWindow))
 {}
 
-FGuid FArchivedConcertSessionTab::GetSessionID() const
-{
-	return InspectedSessionID;
-}
-
 void FArchivedConcertSessionTab::CreateDockContent(const TSharedRef<SDockTab>& InDockTab)
 {
 	SEditableSessionHistory::FMakeSessionHistory MakeSessionHistory = SEditableSessionHistory::FMakeSessionHistory::CreateLambda([this](SSessionHistory::FArguments Arguments)
diff --git a/Engine/Plugins/Developer/Concert/ConcertApp/MultiUserServer/Source/MultiUserServer/Private/Widgets/SessionTabs/Archived/ArchivedConcertSessionTab.h b/Engine/Plugins/Developer/Concert/ConcertApp/MultiUserServer/Source/MultiUserServer/Private/Widgets/SessionTabs/Archived/ArchivedConcertSessionTab.h
--- a/Engine/Plugins/Developer/Concert/ConcertApp/MultiUserServer/Source/MultiUserServer/Private/Widgets/SessionTabs/Archived/ArchivedConcertSessionTab.h
+++ b/Engine/Plugins/Developer/Concert/ConcertApp/MultiUserServer/Source/MultiUserServer/Private/Widgets/SessionTabs/Archived/ArchivedConcertSessionTab.h
@@ -20,7 +20,6 @@ public:
 protected:
 
 	//~ Begin FAbstractConcertSessionTab Interface
-	virtual FGuid GetSessionID() const override;
 	virtual void CreateDockContent(const TSharedRef<SDockTab>& InDockTab) override;
 	virtual void OnOpenTab() override {}
 	//~ End FAbstractConcertSessionTab Interface
diff --git a/Engine/Plugins/Developer/Concert/ConcertApp/MultiUserServer/Source/MultiUserServer/Private/Widgets/SessionTabs/ConcertSessionTabBase.cpp b/Engine/Plugins/Developer/Concert/ConcertApp/MultiUserServer/Source/MultiUserServer/Private/Widgets/SessionTabs/ConcertSessionTabBase.cpp
--- a/Engine/Plugins/Developer/Concert/ConcertApp/MultiUserServer/Source/MultiUserServer/Private/Widgets/SessionTabs/ConcertSessionTabBase.cpp
+++ b/Engine/Plugins/Developer/Concert/ConcertApp/MultiUserServer/Source/MultiUserServer/Private/Widgets/SessionTabs/ConcertSessionTabBase.cpp
@@ -25,10 +25,21 @@ namespace UE::ConcertServerUI::Private
 	}
 }
 
-FConcertSessionTabBase::FConcertSessionTabBase(TSharedRef<IConcertSyncServer> SyncServer)
-	: SyncServer(MoveTemp(SyncServer))
+FConcertSessionTabBase::FConcertSessionTabBase(FGuid InspectedSessionID, TSharedRef<IConcertSyncServer> SyncServer)
+	: InspectedSessionID(MoveTemp(InspectedSessionID))
+	, SyncServer(MoveTemp(SyncServer))
 {}
 
+FConcertSessionTabBase::~FConcertSessionTabBase()
+{
+	const TSharedRef<FGlobalTabmanager>& TabManager = FGlobalTabmanager::Get();
+	const FTabId TabId { *GetTabId() };
+	if (const TSharedPtr<SDockTab> TabStack = TabManager->FindExistingLiveTab(TabId))
+	{
+		TabStack->RequestCloseTab();
+	}
+}
+
 void FConcertSessionTabBase::OpenSessionTab()
 {
 	const TSharedRef<FGlobalTabmanager>& TabManager = FGlobalTabmanager::Get();
@@ -42,7 +53,7 @@ void FConcertSessionTabBase::OpenSessionTab()
 	else
 	{
 		const FTabManager::FLastMajorOrNomadTab Search(ConcertServerTabs::GetSessionBrowserTabId());
-		FGlobalTabmanager::Get()->InsertNewDocumentTab(*GetTabId(), Search, DockTab.ToSharedRef());
+		TabManager->InsertNewDocumentTab(*GetTabId(), Search, DockTab.ToSharedRef());
 
 		OnOpenTab();
 	}
diff --git a/Engine/Plugins/Developer/Concert/ConcertApp/MultiUserServer/Source/MultiUserServer/Private/Widgets/SessionTabs/ConcertSessionTabBase.h b/Engine/Plugins/Developer/Concert/ConcertApp/MultiUserServer/Source/MultiUserServer/Private/Widgets/SessionTabs/ConcertSessionTabBase.h
--- a/Engine/Plugins/Developer/Concert/ConcertApp/MultiUserServer/Source/MultiUserServer/Private/Widgets/SessionTabs/ConcertSessionTabBase.h
+++ b/Engine/Plugins/Developer/Concert/ConcertApp/MultiUserServer/Source/MultiUserServer/Private/Widgets/SessionTabs/ConcertSessionTabBase.h
@@ -13,24 +13,27 @@ class FConcertSessionTabBase
 {
 public:
 
-    FConcertSessionTabBase(TSharedRef<IConcertSyncServer> SyncServer);
-	virtual ~FConcertSessionTabBase() = default;
+    FConcertSessionTabBase(FGuid InspectedSessionID, TSharedRef<IConcertSyncServer> SyncServer);
+	virtual ~FConcertSessionTabBase();
 	
 	/** Opens or draws attention to the given tab */
 	void OpenSessionTab();
 	
 protected:
 
-	virtual FGuid GetSessionID() const = 0;
 	virtual void CreateDockContent(const TSharedRef<SDockTab>& DockTab) = 0;
 	virtual void OnOpenTab() = 0;
 	
+	FGuid GetSessionID() const { return InspectedSessionID; }
 	/** Generates a tab ID for FTabManager::InsertNewDocumentTab */
 	FString GetTabId() const { return GetSessionID().ToString(); }
 	const TSharedRef<IConcertSyncServer>& GetSyncServer() const { return SyncServer; }
 	
 private:
 
+	/** The ID of the session this tab is managing */
+	FGuid InspectedSessionID;
+	
 	/** Used to look up session name */
 	const TSharedRef<IConcertSyncServer> SyncServer;
 	
diff --git a/Engine/Plugins/Developer/Concert/ConcertApp/MultiUserServer/Source/MultiUserServer/Private/Widgets/SessionTabs/Live/LiveConcertSessionTab.cpp b/Engine/Plugins/Developer/Concert/ConcertApp/MultiUserServer/Source/MultiUserServer/Private/Widgets/SessionTabs/Live/LiveConcertSessionTab.cpp
--- a/Engine/Plugins/Developer/Concert/ConcertApp/MultiUserServer/Source/MultiUserServer/Private/Widgets/SessionTabs/Live/LiveConcertSessionTab.cpp
+++ b/Engine/Plugins/Developer/Concert/ConcertApp/MultiUserServer/Source/MultiUserServer/Private/Widgets/SessionTabs/Live/LiveConcertSessionTab.cpp
@@ -11,18 +11,13 @@
 #include "Widgets/Docking/SDockTab.h"
 
 FLiveConcertSessionTab::FLiveConcertSessionTab(TSharedRef<IConcertServerSession> InspectedSession, TSharedRef<IConcertSyncServer> SyncServer, TAttribute<TSharedRef<SWindow>> ConstructUnderWindow)
-	: FConcertSessionTabBase(SyncServer)
+	: FConcertSessionTabBase(InspectedSession->GetSessionInfo().SessionId, SyncServer)
 	, InspectedSession(MoveTemp(InspectedSession))
 	, ConstructUnderWindow(MoveTemp(ConstructUnderWindow))
 	, SessionHistoryController(MakeShared<FLiveServerSessionHistoryController>(InspectedSession, SyncServer))
 	, PackageViewerController(MakeShared<FConcertSessionPackageViewerController>(InspectedSession, SyncServer))
 {}
 
-FGuid FLiveConcertSessionTab::GetSessionID() const
-{
-	return InspectedSession->GetSessionInfo().SessionId;
-}
-
 void FLiveConcertSessionTab::CreateDockContent(const TSharedRef<SDockTab>& InDockTab)
 {
 	const SConcertSessionInspector::FRequiredArgs WidgetArgs
diff --git a/Engine/Plugins/Developer/Concert/ConcertApp/MultiUserServer/Source/MultiUserServer/Private/Widgets/SessionTabs/Live/LiveConcertSessionTab.h b/Engine/Plugins/Developer/Concert/ConcertApp/MultiUserServer/Source/MultiUserServer/Private/Widgets/SessionTabs/Live/LiveConcertSessionTab.h
--- a/Engine/Plugins/Developer/Concert/ConcertApp/MultiUserServer/Source/MultiUserServer/Private/Widgets/SessionTabs/Live/LiveConcertSessionTab.h
+++ b/Engine/Plugins/Developer/Concert/ConcertApp/MultiUserServer/Source/MultiUserServer/Private/Widgets/SessionTabs/Live/LiveConcertSessionTab.h
@@ -26,7 +26,6 @@ public:
 protected:
 
 	//~ Begin FAbstractConcertSessionTab Interface
-	virtual FGuid GetSessionID() const override;
 	virtual void CreateDockContent(const TSharedRef<SDockTab>& InDockTab) override;
 	virtual void OnOpenTab() override;
 	//~ End FAbstractConcertSessionTab Interface
