Delete transactions from archived sessions

diff --git a/Engine/Plugins/Developer/Concert/ConcertApp/MultiUserServer/Source/MultiUserServer/Private/ModalWindowManager.cpp b/Engine/Plugins/Developer/Concert/ConcertApp/MultiUserServer/Source/MultiUserServer/Private/ModalWindowManager.cpp
--- a/Engine/Plugins/Developer/Concert/ConcertApp/MultiUserServer/Source/MultiUserServer/Private/ModalWindowManager.cpp
+++ b/Engine/Plugins/Developer/Concert/ConcertApp/MultiUserServer/Source/MultiUserServer/Private/ModalWindowManager.cpp
@@ -44,9 +44,13 @@ void UE::MultiUserServer::FModalWindowManager::PushWindow(const TSharedRef<SWind
 		check(!WindowStack.Contains(Window));
 		WindowStack.Push(Window);
 	
-		Window->GetOnWindowClosedEvent().AddLambda([this](const TSharedRef<SWindow>&)
+		Window->GetOnWindowClosedEvent().AddLambda([this](const TSharedRef<SWindow>& Window)
 		{
-			PopWindow();
+			// Might request window destruction twice
+			if (WindowStack.Contains(Window))
+			{
+				PopWindow();
+			}
 		});
 	}
 }
diff --git a/Engine/Plugins/Developer/Concert/ConcertApp/MultiUserServer/Source/MultiUserServer/Private/MultiUserServerUserSettings.cpp b/Engine/Plugins/Developer/Concert/ConcertApp/MultiUserServer/Source/MultiUserServer/Private/MultiUserServerUserSettings.cpp
--- a/Engine/Plugins/Developer/Concert/ConcertApp/MultiUserServer/Source/MultiUserServer/Private/MultiUserServerUserSettings.cpp
+++ b/Engine/Plugins/Developer/Concert/ConcertApp/MultiUserServer/Source/MultiUserServer/Private/MultiUserServerUserSettings.cpp
@@ -15,6 +15,10 @@ UMultiUserServerUserSettings::UMultiUserServerUserSettings()
 	{
 		SaveConfig();
 	});
+	OnDeleteActivityDialogColumnVisibility().AddLambda([this](const FColumnVisibilitySnapshot&)
+	{
+		SaveConfig();
+	});
 	OnArchivedActivityBrowserColumnVisibility().AddLambda([this](const FColumnVisibilitySnapshot&)
 	{
 		SaveConfig();
@@ -36,7 +40,7 @@ UMultiUserServerUserSettings::UMultiUserServerUserSettings()
 
 UMultiUserServerUserSettings* UMultiUserServerUserSettings::GetUserSettings()
 {
-	// If we're shutting down GetMutableDefault will return garbage - this function may be called by destructor when this module is unloaded
+	// If we're shutting down GetMutableDefault will return garbage - this function may be called by destructors when this module is unloaded
 	return UE::MultiUserServer::bIsShutdown
 		? nullptr
 		: GetMutableDefault<UMultiUserServerUserSettings>();
diff --git a/Engine/Plugins/Developer/Concert/ConcertApp/MultiUserServer/Source/MultiUserServer/Private/MultiUserServerUserSettings.h b/Engine/Plugins/Developer/Concert/ConcertApp/MultiUserServer/Source/MultiUserServer/Private/MultiUserServerUserSettings.h
--- a/Engine/Plugins/Developer/Concert/ConcertApp/MultiUserServer/Source/MultiUserServer/Private/MultiUserServerUserSettings.h
+++ b/Engine/Plugins/Developer/Concert/ConcertApp/MultiUserServer/Source/MultiUserServer/Private/MultiUserServerUserSettings.h
@@ -26,6 +26,10 @@ public:
 	void SetSessionBrowserColumnVisibility(FColumnVisibilitySnapshot NewValue) { SessionBrowserColumnVisibility = MoveTemp(NewValue); OnSessionBrowserColumnVisibilityChangedEvent.Broadcast(SessionBrowserColumnVisibility); }
 	FOnColumnVisibilitySnapshotChanged& OnSessionBrowserColumnVisibilityChanged() { return OnSessionBrowserColumnVisibilityChangedEvent; }
 
+	const FColumnVisibilitySnapshot& GetDeleteActivityDialogColumnVisibility() const { return DeleteActivityDialogColumnVisibility; }
+	void SetDeleteActivityDialogColumnVisibility(FColumnVisibilitySnapshot NewValue) { DeleteActivityDialogColumnVisibility = MoveTemp(NewValue); OnDeleteActivityDialogColumnVisibilityEvent.Broadcast(DeleteActivityDialogColumnVisibility); }
+	FOnColumnVisibilitySnapshotChanged& OnDeleteActivityDialogColumnVisibility() { return OnDeleteActivityDialogColumnVisibilityEvent; }
+	
 	const FColumnVisibilitySnapshot& GetArchivedActivityBrowserColumnVisibility() const { return ArchivedActivityBrowserColumnVisibility; }
 	void SetArchivedActivityBrowserColumnVisibility(FColumnVisibilitySnapshot NewValue) { ArchivedActivityBrowserColumnVisibility = MoveTemp(NewValue); OnArchivedActivityBrowserColumnVisibilityEvent.Broadcast(ArchivedActivityBrowserColumnVisibility); }
 	FOnColumnVisibilitySnapshotChanged& OnArchivedActivityBrowserColumnVisibility() { return OnArchivedActivityBrowserColumnVisibilityEvent; }
@@ -44,6 +48,10 @@ private:
 	FColumnVisibilitySnapshot SessionBrowserColumnVisibility;
 	FOnColumnVisibilitySnapshotChanged OnSessionBrowserColumnVisibilityChangedEvent;
 
+	UPROPERTY(Config)
+	FColumnVisibilitySnapshot DeleteActivityDialogColumnVisibility;
+	FOnColumnVisibilitySnapshotChanged OnDeleteActivityDialogColumnVisibilityEvent;
+
 	UPROPERTY(Config)
 	FColumnVisibilitySnapshot ArchivedActivityBrowserColumnVisibility;
 	FOnColumnVisibilitySnapshotChanged OnArchivedActivityBrowserColumnVisibilityEvent;
diff --git a/Engine/Plugins/Developer/Concert/ConcertApp/MultiUserServer/Source/MultiUserServer/Private/Widgets/HistoryDeletion/SActivityDependencyView.cpp b/Engine/Plugins/Developer/Concert/ConcertApp/MultiUserServer/Source/MultiUserServer/Private/Widgets/HistoryDeletion/SActivityDependencyView.cpp
new file mode 100644
--- /dev/null
+++ b/Engine/Plugins/Developer/Concert/ConcertApp/MultiUserServer/Source/MultiUserServer/Private/Widgets/HistoryDeletion/SActivityDependencyView.cpp
@@ -0,0 +1,63 @@
+﻿// Copyright Epic Games, Inc. All Rights Reserved.
+
+#include "SActivityDependencyView.h"
+
+#include "HistoryEdition/HistoryAnalysis.h"
+
+#define LOCTEXT_NAMESPACE "UnrealMultiUserUI"
+
+namespace UE::MultiUserServer
+{
+	const FName CustomActivityColumnID("CustomActivityColumn");
+	
+	static FActivityColumn CustomDependencyColumn(const SActivityDependencyView::FCreateActivityColumnWidget& CreateActivityColumnWidget, const FText& CustomActivityColumnLabel)
+	{
+		return FActivityColumn(
+			SHeaderRow::Column(CustomActivityColumnID)
+				.DefaultLabel(CustomActivityColumnLabel)
+				.FixedWidth(20)
+				.ShouldGenerateWidget(true)
+			)
+			.ColumnSortOrder(static_cast<int32>(ConcertSharedSlate::ActivityColumn::EPredefinedColumnOrder::AvatarColor))
+			.GenerateColumnWidget(FActivityColumn::FGenerateColumnWidget::CreateLambda([CreateActivityColumnWidget](const TSharedRef<SConcertSessionActivities>& Owner, const TSharedRef<FConcertSessionActivity>& Activity, SOverlay::FScopedWidgetSlotArguments& Slot)
+			{
+				Slot
+				.VAlign(VAlign_Center)
+				[
+					CreateActivityColumnWidget.Execute(Activity->Activity.ActivityId)
+				];
+			}));
+	}
+}
+
+void SActivityDependencyView::Construct(const FArguments& InArgs, const UE::ConcertSyncCore::FHistoryDeletionRequirements& DeletionRequirements)
+{
+	check(InArgs._CreateSessionHistory.IsBound());
+
+	TSharedPtr<FConcertSessionActivitiesOptions> ViewOptions = MakeShared<FConcertSessionActivitiesOptions>();
+	ViewOptions->bEnableConnectionActivityFiltering = false;
+	ViewOptions->bEnableLockActivityFiltering = false;
+	ViewOptions->bEnableIgnoredActivityFiltering = false;
+
+	TArray<FActivityColumn> Columns { UE::ConcertSharedSlate::ActivityColumn::Operation() };
+	if (InArgs._CreateActivityColumnWidget.IsBound())
+	{
+		Columns.Add(UE::MultiUserServer::CustomDependencyColumn(InArgs._CreateActivityColumnWidget, InArgs._CustomActivityColumnLabel));
+	}
+	
+	ChildSlot
+	[
+		InArgs._CreateSessionHistory.Execute(
+			SSessionHistory::FArguments()
+				.AllowActivity_Lambda([DeletionRequirements](const FConcertSyncActivity& Activity, const TStructOnScope<FConcertSyncActivitySummary>&)
+				{
+					return DeletionRequirements.HardDependencies.Contains(Activity.ActivityId) || DeletionRequirements.PossibleDependencies.Contains(Activity.ActivityId);
+				})
+				.Columns(Columns)
+				.ViewOptions(ViewOptions)
+				.DetailsAreaVisibility(EVisibility::Collapsed)
+			)
+	];
+}
+
+#undef LOCTEXT_NAMESPACE 
diff --git a/Engine/Plugins/Developer/Concert/ConcertApp/MultiUserServer/Source/MultiUserServer/Private/Widgets/HistoryDeletion/SActivityDependencyView.h b/Engine/Plugins/Developer/Concert/ConcertApp/MultiUserServer/Source/MultiUserServer/Private/Widgets/HistoryDeletion/SActivityDependencyView.h
new file mode 100644
--- /dev/null
+++ b/Engine/Plugins/Developer/Concert/ConcertApp/MultiUserServer/Source/MultiUserServer/Private/Widgets/HistoryDeletion/SActivityDependencyView.h
@@ -0,0 +1,37 @@
+﻿// Copyright Epic Games, Inc. All Rights Reserved.
+
+#pragma once
+
+#include "CoreMinimal.h"
+#include "ConcertMessageData.h"
+#include "Session/History/SSessionHistory.h"
+#include "Widgets/DeclarativeSyntaxSupport.h"
+#include "Widgets/SCompoundWidget.h"
+
+namespace UE
+{
+	namespace ConcertSyncCore
+	{
+		struct FHistoryDeletionRequirements;
+	}
+}
+
+/** A view for displaying activity dependencies when deleting a dependency. */
+class SActivityDependencyView : public SCompoundWidget
+{
+public:
+
+	DECLARE_DELEGATE_RetVal_OneParam(TSharedRef<SSessionHistory>, FCreateSessionHistory, const SSessionHistory::FArguments& /*InArgs*/)
+	DECLARE_DELEGATE_RetVal_OneParam(TSharedRef<SWidget>, FCreateActivityColumnWidget, FActivityID /*ActivityID*/);
+
+    SLATE_BEGIN_ARGS(SActivityDependencyView)
+    {}
+		/** Required. A function that will create a session history widget. */
+		SLATE_EVENT(FCreateSessionHistory, CreateSessionHistory)
+		/** Optional. If specified, this will create a widget in the first column using this function. */
+		SLATE_EVENT(FCreateActivityColumnWidget, CreateActivityColumnWidget)
+		SLATE_ARGUMENT(FText, CustomActivityColumnLabel)
+    SLATE_END_ARGS()
+	
+	void Construct(const FArguments& InArgs, const UE::ConcertSyncCore::FHistoryDeletionRequirements& DeletionRequirements);
+};
diff --git a/Engine/Plugins/Developer/Concert/ConcertApp/MultiUserServer/Source/MultiUserServer/Private/Widgets/HistoryDeletion/SDeleteActivityDependenciesDialog.cpp b/Engine/Plugins/Developer/Concert/ConcertApp/MultiUserServer/Source/MultiUserServer/Private/Widgets/HistoryDeletion/SDeleteActivityDependenciesDialog.cpp
new file mode 100644
--- /dev/null
+++ b/Engine/Plugins/Developer/Concert/ConcertApp/MultiUserServer/Source/MultiUserServer/Private/Widgets/HistoryDeletion/SDeleteActivityDependenciesDialog.cpp
@@ -0,0 +1,123 @@
+﻿// Copyright Epic Games, Inc. All Rights Reserved.
+
+#include "SDeleteActivityDependenciesDialog.h"
+
+#include "ConcertSyncSessionDatabase.h"
+#include "SActivityDependencyView.h"
+
+#include "Widgets/Input/SCheckBox.h"
+#include "Widgets/SBoxPanel.h"
+#include "Widgets/SessionTabs/Archived/ArchivedSessionHistoryController.h"
+#include "Widgets/Text/STextBlock.h"
+
+#define LOCTEXT_NAMESPACE "UnrealMultiUserUI"
+
+void SDeleteActivityDependenciesDialog::Construct(const FArguments& InArgs, const FGuid& SessionId, const TSharedRef<IConcertSyncServer>& SyncServer, UE::ConcertSyncCore::FHistoryDeletionRequirements InDeletionRequirements)
+{
+	DeletionRequirements = MoveTemp(InDeletionRequirements);
+	OnConfirmDeletionFunc = InArgs._OnConfirmDeletion;
+
+	check(OnConfirmDeletionFunc.IsBound());
+	
+	SCustomDialog::Construct(
+		SCustomDialog::FArguments()
+		.Title(LOCTEXT("RemoveActivityTitle", "Remove activity"))
+		.Buttons({
+			FButton(LOCTEXT("DeleteActivity.ConfirmButtonLabel", "Delete"))
+				.SetOnClicked(FSimpleDelegate::CreateSP(this, &SDeleteActivityDependenciesDialog::OnConfirmPressed)),
+			FButton(LOCTEXT("DeleteActivity.CancelButtonLabel", "Cancel"))
+				.SetPrimary(true)
+		})
+		.Content()
+		[
+			SNew(SVerticalBox)
+			
+			+SVerticalBox::Slot()
+			.HAlign(HAlign_Fill)
+			.VAlign(VAlign_Top)
+			.AutoHeight()
+			.Padding(0, 5, 0, 10)
+			[
+				SNew(STextBlock)
+				.Text(LOCTEXT("DeleteActivity.HeaderBaseText", "Review the activities that will be deleted:"))
+			]
+
+			+SVerticalBox::Slot()
+			.HAlign(HAlign_Fill)
+			.VAlign(VAlign_Fill)
+			[
+				CreateBody(SessionId, SyncServer, InDeletionRequirements)
+			]
+		]
+	);
+}
+
+TSharedRef<SWidget> SDeleteActivityDependenciesDialog::CreateBody(const FGuid& InSessionId, const TSharedRef<IConcertSyncServer>& InSyncServer, UE::ConcertSyncCore::FHistoryDeletionRequirements InDeletionRequirements)
+{
+	for (FActivityID HardDependency : DeletionRequirements.HardDependencies)
+	{
+		ActivitiesToDelete.Add(HardDependency, true);
+	}
+	for (FActivityID HardDependency : DeletionRequirements.PossibleDependencies)
+	{
+		ActivitiesToDelete.Add(HardDependency, false);
+	}
+	
+	return SNew(SActivityDependencyView, DeletionRequirements)
+		.CreateSessionHistory_Lambda([this, InSessionId, InSyncServer](const SSessionHistory::FArguments& InArgs)
+		{
+			const bool bOnlyCalledOnce = !FilteredSessionHistoryController.IsValid();
+			check(bOnlyCalledOnce);
+			
+			FilteredSessionHistoryController = UE::MultiUserServer::CreateForDeletionDialog(
+				InSessionId,
+				InSyncServer,
+				InArgs
+				);
+			return FilteredSessionHistoryController->GetSessionHistory();
+		})
+		.CreateActivityColumnWidget_Lambda([this](FActivityID ActivityID)
+		{
+			const bool bIsHardDependency = DeletionRequirements.HardDependencies.Contains(ActivityID);
+			return SNew(SCheckBox)
+				.IsEnabled(!bIsHardDependency)
+				.IsChecked_Lambda([this, ActivityID]()
+				{
+					return ActivitiesToDelete[ActivityID] ? ECheckBoxState::Checked : ECheckBoxState::Unchecked;
+				})
+				.OnCheckStateChanged_Lambda([this, ActivityID](ECheckBoxState NewState)
+				{
+					ActivitiesToDelete[ActivityID] = NewState == ECheckBoxState::Checked;
+				})
+				.ToolTipText_Lambda([this, ActivityID, bIsHardDependency]()
+				{
+					if (bIsHardDependency)
+					{
+						return LOCTEXT("DeleteActivity.CheckBox.TooltipDisabled", "This is a hard dependency and must be deleted.");
+					}
+
+					const bool bWillBeDeleted = ActivitiesToDelete[ActivityID];
+					if (bWillBeDeleted)
+					{
+						return LOCTEXT("DeleteActivity.CheckBox.TooltipEnabled", "This is a possible dependency.\nWill be deleted.");
+					}
+					return LOCTEXT("DeleteActivity.CheckBox.TooltipEnabled", "This is a possible dependency.\nWill not be deleted.");
+				});
+		});
+}
+
+void SDeleteActivityDependenciesDialog::OnConfirmPressed()
+{
+	for (const FActivityID& ActivityID : TSet<FActivityID>(DeletionRequirements.PossibleDependencies))
+	{
+		const bool bShouldDelete = ActivitiesToDelete[ActivityID];
+		if (!bShouldDelete)
+		{
+			DeletionRequirements.PossibleDependencies.Remove(ActivityID);
+		}
+	}
+	
+	OnConfirmDeletionFunc.Execute(DeletionRequirements);
+}
+
+#undef LOCTEXT_NAMESPACE
diff --git a/Engine/Plugins/Developer/Concert/ConcertApp/MultiUserServer/Source/MultiUserServer/Private/Widgets/HistoryDeletion/SDeleteActivityDependenciesDialog.h b/Engine/Plugins/Developer/Concert/ConcertApp/MultiUserServer/Source/MultiUserServer/Private/Widgets/HistoryDeletion/SDeleteActivityDependenciesDialog.h
new file mode 100644
--- /dev/null
+++ b/Engine/Plugins/Developer/Concert/ConcertApp/MultiUserServer/Source/MultiUserServer/Private/Widgets/HistoryDeletion/SDeleteActivityDependenciesDialog.h
@@ -0,0 +1,48 @@
+﻿// Copyright Epic Games, Inc. All Rights Reserved.
+
+#pragma once
+
+#include "CoreMinimal.h"
+#include "ConcertMessageData.h"
+#include "IConcertSyncServer.h"
+
+#include "Dialog/SCustomDialog.h"
+#include "HistoryEdition/HistoryAnalysis.h"
+#include "Widgets/DeclarativeSyntaxSupport.h"
+
+class FArchivedSessionHistoryController;
+
+/** Displayed when a user asks to delete an activity. */
+class SDeleteActivityDependenciesDialog : public SCustomDialog
+{
+public:
+
+	DECLARE_DELEGATE_OneParam(FConfirmDeletion, const UE::ConcertSyncCore::FHistoryDeletionRequirements& /*SelectedRequirements*/)
+
+	SLATE_BEGIN_ARGS(SDeleteActivityDependenciesDialog)
+	{}
+		/** Called when the user confirms the deletion of the activities.*/
+		SLATE_EVENT(FConfirmDeletion, OnConfirmDeletion)
+	SLATE_END_ARGS()
+
+	/**
+	 * @param InDeletionRequirements Specifies which activities must be deleted and which are optional.
+	 */
+	void Construct(const FArguments& InArgs, const FGuid& SessionId, const TSharedRef<IConcertSyncServer>& SyncServer, UE::ConcertSyncCore::FHistoryDeletionRequirements InDeletionRequirements);
+
+private:
+	
+	UE::ConcertSyncCore::FHistoryDeletionRequirements DeletionRequirements;
+	FConfirmDeletion OnConfirmDeletionFunc;
+	
+	/** Maps each activity with a possible dependency to whether it will be deleted or not. */
+	TMap<FActivityID, bool> ActivitiesToDelete;
+	
+	/** Filters out all activities that are not part of the deletion requirements. */
+	TSharedPtr<FArchivedSessionHistoryController> FilteredSessionHistoryController;
+
+	TSharedRef<SWidget> CreateBody(const FGuid& InSessionId, const TSharedRef<IConcertSyncServer>& InSyncServer, UE::ConcertSyncCore::FHistoryDeletionRequirements InDeletionRequirements);
+
+	void OnConfirmPressed();
+};
+
diff --git a/Engine/Plugins/Developer/Concert/ConcertApp/MultiUserServer/Source/MultiUserServer/Private/Widgets/SessionTabs/Archived/ArchivedConcertSessionTab.cpp b/Engine/Plugins/Developer/Concert/ConcertApp/MultiUserServer/Source/MultiUserServer/Private/Widgets/SessionTabs/Archived/ArchivedConcertSessionTab.cpp
--- a/Engine/Plugins/Developer/Concert/ConcertApp/MultiUserServer/Source/MultiUserServer/Private/Widgets/SessionTabs/Archived/ArchivedConcertSessionTab.cpp
+++ b/Engine/Plugins/Developer/Concert/ConcertApp/MultiUserServer/Source/MultiUserServer/Private/Widgets/SessionTabs/Archived/ArchivedConcertSessionTab.cpp
@@ -3,11 +3,29 @@
 #include "Widgets/SessionTabs/Archived/ArchivedConcertSessionTab.h"
 
 #include "ArchivedSessionHistoryController.h"
+#include "ConcertLogGlobal.h"
+#include "ModalWindowManager.h"
+#include "MultiUserServerModule.h"
+
+#include "Algo/AllOf.h"
+#include "Algo/Transform.h"
+
+#include "Dialog/SMessageDialog.h"
+
+#include "HistoryEdition/ActivityNode.h"
+#include "HistoryEdition/DependencyGraphBuilder.h"
+#include "HistoryEdition/HistoryAnalysis.h"
+#include "HistoryEdition/HistoryDeletion.h"
+
 #include "Session/History/SEditableSessionHistory.h"
 #include "Session/History/SSessionHistory.h"
+
+#include "Widgets/HistoryDeletion/SDeleteActivityDependenciesDialog.h"
 #include "Widgets/SessionTabs/Archived/SConcertArchivedSessionInspector.h"
 #include "Widgets/StatusBar/SConcertStatusBar.h"
 
+#define LOCTEXT_NAMESPACE "UnrealMultiUserUI"
+
 FArchivedConcertSessionTab::FArchivedConcertSessionTab(FGuid InspectedSessionID, TSharedRef<IConcertSyncServer> SyncServer, TAttribute<TSharedRef<SWindow>> ConstructUnderWindow)
 	: FConcertSessionTabBase(InspectedSessionID, SyncServer)
 	, InspectedSessionID(MoveTemp(InspectedSessionID))
@@ -20,7 +38,7 @@ void FArchivedConcertSessionTab::CreateDockContent(const TSharedRef<SDockTab>& I
 	SEditableSessionHistory::FMakeSessionHistory MakeSessionHistory = SEditableSessionHistory::FMakeSessionHistory::CreateLambda([this](SSessionHistory::FArguments Arguments)
 	{
 		checkf(!HistoryController.IsValid(), TEXT("Called more than once"));
-		HistoryController = MakeShared<FArchivedSessionHistoryController>(InspectedSessionID, SyncServer, MoveTemp(Arguments));
+		HistoryController = UE::MultiUserServer::CreateForInspector(InspectedSessionID, SyncServer, MoveTemp(Arguments));
 		return HistoryController->GetSessionHistory();
 	});
 	
@@ -38,13 +56,67 @@ void FArchivedConcertSessionTab::CreateDockContent(const TSharedRef<SDockTab>& I
 		);
 }
 
-void FArchivedConcertSessionTab::OnRequestDeleteActivity(const TSharedRef<FConcertSessionActivity>& DeleteActivity) const
+void FArchivedConcertSessionTab::OnRequestDeleteActivity(const TSet<TSharedRef<FConcertSessionActivity>>& ActivitiesToDelete) const
 {
-	// TODO:
+	using namespace UE::ConcertSyncCore;
+	
+	if (const TOptional<FConcertSyncSessionDatabaseNonNullPtr> SessionDatabase = SyncServer->GetArchivedSessionDatabase(InspectedSessionID))
+	{
+		TSet<FActivityID> RequestedForDelete;
+		Algo::Transform(ActivitiesToDelete, RequestedForDelete, [](const TSharedRef<FConcertSessionActivity>& Activity)
+		{
+			return Activity->Activity.ActivityId;
+		});
+		const FActivityDependencyGraph DependencyGraph = BuildDependencyGraphFrom(*SessionDatabase);
+		FHistoryDeletionRequirements DeletionRequirements = AnalyseActivityDeletion(RequestedForDelete, DependencyGraph, true);
+
+		TWeakPtr<const FArchivedConcertSessionTab> WeakTabThis = SharedThis(this);
+		TSharedRef<SDeleteActivityDependenciesDialog> Dialog = SNew(SDeleteActivityDependenciesDialog, InspectedSessionID, SyncServer, MoveTemp(DeletionRequirements))
+			.OnConfirmDeletion_Lambda([WeakTabThis](const FHistoryDeletionRequirements& SelectedRequirements)
+			{
+				// Because the dialog is non-modal, the user may have closed the program in the mean time
+				if (const TSharedPtr<const FArchivedConcertSessionTab> PinnedThis = WeakTabThis.Pin())
+				{
+					const FDeleteSessionErrorResult ErrorResult = DeleteActivitiesInArchivedSession(PinnedThis->SyncServer->GetConcertServer(), PinnedThis->InspectedSessionID, CombineRequirements(SelectedRequirements));
+					if (ErrorResult.HadError())
+					{
+						UE_LOG(LogConcert, Error, TEXT("Failed to delete activities from session %s: %s"), *PinnedThis->InspectedSessionID.ToString(), *ErrorResult.ErrorMessage->ToString());
+						
+						const TSharedRef<SMessageDialog> ErrorDialog = SNew(SMessageDialog)
+							.Title(LOCTEXT("ErrorDeletingSessions", "Error deleting sessions"))
+							.Message(*ErrorResult.ErrorMessage)
+							.Buttons({
+								SMessageDialog::FButton(LOCTEXT("Ok", "Ok"))
+								.SetPrimary(true)
+							});
+						ErrorDialog->Show();
+					}
+					else
+					{
+						// The list needs to be refreshed after the delete operation
+						PinnedThis->HistoryController->ReloadActivities();
+					}
+				}
+			});
+		
+		UE::MultiUserServer::FConcertServerUIModule::Get()
+			.GetModalWindowManager()
+			->ShowFakeModalWindow(Dialog);
+	}
 }
 
-bool FArchivedConcertSessionTab::CanDeleteActivity(const TSharedRef<FConcertSessionActivity>& DeleteActivity) const
+FCanDeleteActivitiesResult FArchivedConcertSessionTab::CanDeleteActivity(const TSet<TSharedRef<FConcertSessionActivity>>& ActivitiesToDelete) const
 {
-	// TODO:
-	return false;
+	const bool bOnlyPackagesAndTransactions =  Algo::AllOf(ActivitiesToDelete, [](const TSharedRef<FConcertSessionActivity>& Activity)
+	{
+		return Activity->Activity.EventType == EConcertSyncActivityEventType::Package || Activity->Activity.EventType == EConcertSyncActivityEventType::Transaction;
+	});
+	if (!bOnlyPackagesAndTransactions)
+	{
+		return FCanDeleteActivitiesResult::No(LOCTEXT("CanDeleteActivity.OnlyPackagesAndTransactionsReason", "Only package and transaction activities can be deleted (the current selection includes other activity types)."));
+	}
+
+	return FCanDeleteActivitiesResult::Yes();
 }
+
+#undef LOCTEXT_NAMESPACE
\ No newline at end of file
diff --git a/Engine/Plugins/Developer/Concert/ConcertApp/MultiUserServer/Source/MultiUserServer/Private/Widgets/SessionTabs/Archived/ArchivedConcertSessionTab.h b/Engine/Plugins/Developer/Concert/ConcertApp/MultiUserServer/Source/MultiUserServer/Private/Widgets/SessionTabs/Archived/ArchivedConcertSessionTab.h
--- a/Engine/Plugins/Developer/Concert/ConcertApp/MultiUserServer/Source/MultiUserServer/Private/Widgets/SessionTabs/Archived/ArchivedConcertSessionTab.h
+++ b/Engine/Plugins/Developer/Concert/ConcertApp/MultiUserServer/Source/MultiUserServer/Private/Widgets/SessionTabs/Archived/ArchivedConcertSessionTab.h
@@ -3,8 +3,8 @@
 #pragma once
 
 #include "CoreMinimal.h"
+#include "Session/History/SEditableSessionHistory.h"
 #include "Widgets/SessionTabs/ConcertSessionTabBase.h"
-#include "Types/SlateAttribute.h"
 
 struct FConcertSessionActivity;
 class FArchivedSessionHistoryController;
@@ -40,6 +40,6 @@ private:
 	/** Displays session */
 	TSharedPtr<SConcertArchivedSessionInspector> Inspector;
 
-	void OnRequestDeleteActivity(const TSharedRef<FConcertSessionActivity>& DeleteActivity) const;
-	bool CanDeleteActivity(const TSharedRef<FConcertSessionActivity>& DeleteActivity) const;
+	void OnRequestDeleteActivity(const TSet<TSharedRef<FConcertSessionActivity>>& ActivitiesToDelete) const;
+	FCanDeleteActivitiesResult CanDeleteActivity(const TSet<TSharedRef<FConcertSessionActivity>>& ActivitiesToDelete) const;
 };
diff --git a/Engine/Plugins/Developer/Concert/ConcertApp/MultiUserServer/Source/MultiUserServer/Private/Widgets/SessionTabs/Archived/ArchivedSessionHistoryController.cpp b/Engine/Plugins/Developer/Concert/ConcertApp/MultiUserServer/Source/MultiUserServer/Private/Widgets/SessionTabs/Archived/ArchivedSessionHistoryController.cpp
--- a/Engine/Plugins/Developer/Concert/ConcertApp/MultiUserServer/Source/MultiUserServer/Private/Widgets/SessionTabs/Archived/ArchivedSessionHistoryController.cpp
+++ b/Engine/Plugins/Developer/Concert/ConcertApp/MultiUserServer/Source/MultiUserServer/Private/Widgets/SessionTabs/Archived/ArchivedSessionHistoryController.cpp
@@ -4,33 +4,11 @@
 
 #include "MultiUserServerUserSettings.h"
 
-namespace UE::MultiUserServer::Private
-{
-	static SSessionHistory::FArguments MakeArchivedTabSessionHistoryArguments(SSessionHistory::FArguments&& InArgs)
-	{
-		return InArgs
-			.ColumnVisibilitySnapshot(UMultiUserServerUserSettings::GetUserSettings()->GetArchivedActivityBrowserColumnVisibility())
-			.SaveColumnVisibilitySnapshot_Lambda([](const FColumnVisibilitySnapshot& Snapshot)
-			{
-				UMultiUserServerUserSettings::GetUserSettings()->SetArchivedActivityBrowserColumnVisibility(Snapshot);
-			});
-	}
-}
-
 FArchivedSessionHistoryController::FArchivedSessionHistoryController(FGuid SessionId, TSharedRef<IConcertSyncServer> SyncServer, SSessionHistory::FArguments Arguments)
-	: FServerSessionHistoryControllerBase(MoveTemp(SessionId), UE::MultiUserServer::Private::MakeArchivedTabSessionHistoryArguments(MoveTemp(Arguments)))
+	: FServerSessionHistoryControllerBase(MoveTemp(SessionId), MoveTemp(Arguments))
 	, SyncServer(MoveTemp(SyncServer))
 {
 	ReloadActivities();
-	UMultiUserServerUserSettings::GetUserSettings()->OnArchivedActivityBrowserColumnVisibility().AddRaw(this, &FArchivedSessionHistoryController::OnActivityListColumnVisibilitySettingsUpdated);
-}
-
-FArchivedSessionHistoryController::~FArchivedSessionHistoryController()
-{
-	if (UMultiUserServerUserSettings* Settings = UMultiUserServerUserSettings::GetUserSettings(); IsValid(Settings))
-	{
-		Settings->OnArchivedActivityBrowserColumnVisibility().RemoveAll(this);
-	}
 }
 
 TOptional<FConcertSyncSessionDatabaseNonNullPtr> FArchivedSessionHistoryController::GetSessionDatabase(const FGuid& InSessionId) const
@@ -38,7 +16,87 @@ TOptional<FConcertSyncSessionDatabaseNonNullPtr> FArchivedSessionHistoryControll
 	return SyncServer->GetArchivedSessionDatabase(InSessionId);
 }
 
-void FArchivedSessionHistoryController::OnActivityListColumnVisibilitySettingsUpdated(const FColumnVisibilitySnapshot& NewValue)
+namespace UE::MultiUserServer
 {
-	GetSessionHistory()->OnColumnVisibilitySettingsChanged(NewValue);
+	namespace Private
+	{
+		static SSessionHistory::FArguments MakeArgumentsForInspector(SSessionHistory::FArguments&& InArgs)
+		{
+			return InArgs
+				.ColumnVisibilitySnapshot(UMultiUserServerUserSettings::GetUserSettings()->GetArchivedActivityBrowserColumnVisibility())
+				.SaveColumnVisibilitySnapshot_Lambda([](const FColumnVisibilitySnapshot& Snapshot)
+				{
+					UMultiUserServerUserSettings::GetUserSettings()->SetArchivedActivityBrowserColumnVisibility(Snapshot);
+				});
+		}
+
+		class FInspectorSessionHistoryController : public FArchivedSessionHistoryController
+		{
+		public:
+			FInspectorSessionHistoryController(FGuid SessionId, TSharedRef<IConcertSyncServer> SyncServer, SSessionHistory::FArguments Arguments)
+				: FArchivedSessionHistoryController(MoveTemp(SessionId), MoveTemp(SyncServer), MakeArgumentsForInspector(MoveTemp(Arguments)))
+			{
+				UMultiUserServerUserSettings::GetUserSettings()->OnArchivedActivityBrowserColumnVisibility().AddRaw(this, &FInspectorSessionHistoryController::OnActivityListColumnVisibilitySettingsUpdated);
+			}
+		
+			virtual ~FInspectorSessionHistoryController() override
+			{
+				if (UMultiUserServerUserSettings* Settings = UMultiUserServerUserSettings::GetUserSettings(); IsValid(Settings))
+				{
+					Settings->OnArchivedActivityBrowserColumnVisibility().RemoveAll(this);
+				}
+			}
+		
+			void OnActivityListColumnVisibilitySettingsUpdated(const FColumnVisibilitySnapshot& NewValue)
+			{
+				GetSessionHistory()->OnColumnVisibilitySettingsChanged(NewValue);
+			}
+		};
+		
+		static SSessionHistory::FArguments MakeArgumentsForDeleteDialog(SSessionHistory::FArguments&& InArgs)
+		{
+			return InArgs
+				.ColumnVisibilitySnapshot(UMultiUserServerUserSettings::GetUserSettings()->GetDeleteActivityDialogColumnVisibility())
+				.SaveColumnVisibilitySnapshot_Lambda([](const FColumnVisibilitySnapshot& Snapshot)
+				{
+					UMultiUserServerUserSettings::GetUserSettings()->SetDeleteActivityDialogColumnVisibility(Snapshot);
+				});
+		}
+
+		class FDeleteDialogSessionHistoryController : public FArchivedSessionHistoryController
+		{
+		public:
+			
+			FDeleteDialogSessionHistoryController(FGuid SessionId, TSharedRef<IConcertSyncServer> SyncServer, SSessionHistory::FArguments Arguments)
+				: FArchivedSessionHistoryController(MoveTemp(SessionId), MoveTemp(SyncServer), MakeArgumentsForDeleteDialog(MoveTemp(Arguments)))
+			{
+				UMultiUserServerUserSettings::GetUserSettings()->OnDeleteActivityDialogColumnVisibility().AddRaw(this, &FDeleteDialogSessionHistoryController::OnActivityListColumnVisibilitySettingsUpdated);
+			}
+		
+			virtual ~FDeleteDialogSessionHistoryController() override
+			{
+				if (UMultiUserServerUserSettings* Settings = UMultiUserServerUserSettings::GetUserSettings(); IsValid(Settings))
+				{
+					Settings->OnDeleteActivityDialogColumnVisibility().RemoveAll(this);
+				}
+			}
+		
+			void OnActivityListColumnVisibilitySettingsUpdated(const FColumnVisibilitySnapshot& NewValue)
+			{
+				GetSessionHistory()->OnColumnVisibilitySettingsChanged(NewValue);
+			}
+		};
+	}
+
+	TSharedPtr<FArchivedSessionHistoryController> CreateForInspector(FGuid SessionId, TSharedRef<IConcertSyncServer> SyncServer, SSessionHistory::FArguments Arguments)
+	{
+		return MakeShared<Private::FInspectorSessionHistoryController>(MoveTemp(SessionId), MoveTemp(SyncServer), MoveTemp(Arguments));
+	}
+
+	TSharedPtr<FArchivedSessionHistoryController> CreateForDeletionDialog(FGuid SessionId, TSharedRef<IConcertSyncServer> SyncServer, SSessionHistory::FArguments Arguments)
+	{
+		return MakeShared<Private::FDeleteDialogSessionHistoryController>(MoveTemp(SessionId), MoveTemp(SyncServer), MoveTemp(Arguments));
+	}
 }
+
+
diff --git a/Engine/Plugins/Developer/Concert/ConcertApp/MultiUserServer/Source/MultiUserServer/Private/Widgets/SessionTabs/Archived/ArchivedSessionHistoryController.h b/Engine/Plugins/Developer/Concert/ConcertApp/MultiUserServer/Source/MultiUserServer/Private/Widgets/SessionTabs/Archived/ArchivedSessionHistoryController.h
--- a/Engine/Plugins/Developer/Concert/ConcertApp/MultiUserServer/Source/MultiUserServer/Private/Widgets/SessionTabs/Archived/ArchivedSessionHistoryController.h
+++ b/Engine/Plugins/Developer/Concert/ConcertApp/MultiUserServer/Source/MultiUserServer/Private/Widgets/SessionTabs/Archived/ArchivedSessionHistoryController.h
@@ -5,14 +5,12 @@
 #include "CoreMinimal.h"
 #include "Widgets/SessionTabs/ServerSessionHistoryControllerBase.h"
 
+/** Manages SSessionHistory by using IConcertSyncServer::GetArchivedSessionDatabase for retrieving the session database. */
 class FArchivedSessionHistoryController : public FServerSessionHistoryControllerBase
 {
-public:
+protected: // Protected to encourage the use the below functions (CreateForInspector, etc.) instead of MakeShared
 	
 	FArchivedSessionHistoryController(FGuid SessionId, TSharedRef<IConcertSyncServer> SyncServer, SSessionHistory::FArguments Arguments);
-	virtual ~FArchivedSessionHistoryController() override;
-	
-protected:
 
 	//~ Begin FServerSessionHistoryControllerBase Interface
 	virtual TOptional<FConcertSyncSessionDatabaseNonNullPtr> GetSessionDatabase(const FGuid& InSessionId) const override;
@@ -21,6 +19,13 @@ protected:
 private:
 
 	const TSharedRef<IConcertSyncServer> SyncServer;
-
-	void OnActivityListColumnVisibilitySettingsUpdated(const FColumnVisibilitySnapshot& NewValue);
 };
+
+namespace UE::MultiUserServer
+{
+	/** Creates a controller that is supposed to use UMultiUserServerUserSettings::GetArchivedActivityBrowserColumnVisibility  */
+	TSharedPtr<FArchivedSessionHistoryController> CreateForInspector(FGuid SessionId, TSharedRef<IConcertSyncServer> SyncServer, SSessionHistory::FArguments Arguments);
+
+	/** Creates a controller that is supposed to use  UMultiUserServerUserSettings::GetDeleteActivityDialogColumnVisibility*/
+	TSharedPtr<FArchivedSessionHistoryController> CreateForDeletionDialog(FGuid SessionId, TSharedRef<IConcertSyncServer> SyncServer, SSessionHistory::FArguments Arguments);
+}
diff --git a/Engine/Plugins/Developer/Concert/ConcertApp/MultiUserServer/Source/MultiUserServer/Private/Widgets/SessionTabs/Archived/SConcertArchivedSessionInspector.cpp b/Engine/Plugins/Developer/Concert/ConcertApp/MultiUserServer/Source/MultiUserServer/Private/Widgets/SessionTabs/Archived/SConcertArchivedSessionInspector.cpp
--- a/Engine/Plugins/Developer/Concert/ConcertApp/MultiUserServer/Source/MultiUserServer/Private/Widgets/SessionTabs/Archived/SConcertArchivedSessionInspector.cpp
+++ b/Engine/Plugins/Developer/Concert/ConcertApp/MultiUserServer/Source/MultiUserServer/Private/Widgets/SessionTabs/Archived/SConcertArchivedSessionInspector.cpp
@@ -74,8 +74,8 @@ TSharedRef<SWidget> SConcertArchivedSessionInspector::CreateTabs(const FArgument
 TSharedRef<SDockTab> SConcertArchivedSessionInspector::SpawnActivityHistory(
 	const FSpawnTabArgs& Args,
 	SEditableSessionHistory::FMakeSessionHistory MakeSessionHistory,
-	SEditableSessionHistory::FCanDeleteActivity CanDeleteActivity,
-	SEditableSessionHistory::FRequestDeleteActivity DeleteActivity)
+	SEditableSessionHistory::FCanDeleteActivities CanDeleteActivity,
+	SEditableSessionHistory::FRequestDeleteActivities DeleteActivity)
 {
 	SessionHistory = SNew(SEditableSessionHistory)
 		.MakeSessionHistory(MoveTemp(MakeSessionHistory))
diff --git a/Engine/Plugins/Developer/Concert/ConcertApp/MultiUserServer/Source/MultiUserServer/Private/Widgets/SessionTabs/Archived/SConcertArchivedSessionInspector.h b/Engine/Plugins/Developer/Concert/ConcertApp/MultiUserServer/Source/MultiUserServer/Private/Widgets/SessionTabs/Archived/SConcertArchivedSessionInspector.h
--- a/Engine/Plugins/Developer/Concert/ConcertApp/MultiUserServer/Source/MultiUserServer/Private/Widgets/SessionTabs/Archived/SConcertArchivedSessionInspector.h
+++ b/Engine/Plugins/Developer/Concert/ConcertApp/MultiUserServer/Source/MultiUserServer/Private/Widgets/SessionTabs/Archived/SConcertArchivedSessionInspector.h
@@ -25,8 +25,8 @@ public:
 		SLATE_ARGUMENT(TSharedPtr<SDockTab>, ConstructUnderMajorTab)
 		SLATE_ARGUMENT(TSharedPtr<SWindow>, ConstructUnderWindow)
 		SLATE_EVENT(SEditableSessionHistory::FMakeSessionHistory, MakeSessionHistory)
-		SLATE_EVENT(SEditableSessionHistory::FCanDeleteActivity, CanDeleteActivity)
-		SLATE_EVENT(SEditableSessionHistory::FRequestDeleteActivity, DeleteActivity)
+		SLATE_EVENT(SEditableSessionHistory::FCanDeleteActivities, CanDeleteActivity)
+		SLATE_EVENT(SEditableSessionHistory::FRequestDeleteActivities, DeleteActivity)
 		SLATE_NAMED_SLOT(FArguments, StatusBar)
 	SLATE_END_ARGS()
 
@@ -42,7 +42,7 @@ private:
 	TSharedRef<SWidget> CreateTabs(const FArguments& InArgs);
 	TSharedRef<SDockTab> SpawnActivityHistory(const FSpawnTabArgs& Args,
 		SEditableSessionHistory::FMakeSessionHistory FMakeSessionHistory,
-		SEditableSessionHistory::FCanDeleteActivity CanDeleteActivity,
-		SEditableSessionHistory::FRequestDeleteActivity DeleteActivity
+		SEditableSessionHistory::FCanDeleteActivities CanDeleteActivity,
+		SEditableSessionHistory::FRequestDeleteActivities DeleteActivity
 		);
 };
diff --git a/Engine/Plugins/Developer/Concert/ConcertApp/MultiUserServer/Source/MultiUserServer/Private/Widgets/SessionTabs/ConcertSessionTabBase.h b/Engine/Plugins/Developer/Concert/ConcertApp/MultiUserServer/Source/MultiUserServer/Private/Widgets/SessionTabs/ConcertSessionTabBase.h
--- a/Engine/Plugins/Developer/Concert/ConcertApp/MultiUserServer/Source/MultiUserServer/Private/Widgets/SessionTabs/ConcertSessionTabBase.h
+++ b/Engine/Plugins/Developer/Concert/ConcertApp/MultiUserServer/Source/MultiUserServer/Private/Widgets/SessionTabs/ConcertSessionTabBase.h
@@ -9,7 +9,7 @@ class SDockTab;
 class SWindow;
 
 /** Shared functionality for a tab that manages a session */
-class FConcertSessionTabBase
+class FConcertSessionTabBase : public TSharedFromThis<FConcertSessionTabBase>
 {
 public:
 
diff --git a/Engine/Plugins/Developer/Concert/ConcertApp/MultiUserServer/Source/MultiUserServer/Private/Widgets/SessionTabs/Live/LiveServerSessionHistoryController.h b/Engine/Plugins/Developer/Concert/ConcertApp/MultiUserServer/Source/MultiUserServer/Private/Widgets/SessionTabs/Live/LiveServerSessionHistoryController.h
--- a/Engine/Plugins/Developer/Concert/ConcertApp/MultiUserServer/Source/MultiUserServer/Private/Widgets/SessionTabs/Live/LiveServerSessionHistoryController.h
+++ b/Engine/Plugins/Developer/Concert/ConcertApp/MultiUserServer/Source/MultiUserServer/Private/Widgets/SessionTabs/Live/LiveServerSessionHistoryController.h
@@ -5,11 +5,13 @@
 #include "CoreMinimal.h"
 #include "Widgets/SessionTabs/ServerSessionHistoryControllerBase.h"
 
-struct FConcertSyncActivity;
 class FConcertSyncSessionDatabase;
 class IConcertSyncServer;
 class IConcertServerSession;
 
+struct FConcertSyncActivity;
+
+/** Manages SSessionHistory by using IConcertSyncServer::GetLiveSessionDatabase for retrieving the session database. */
 class FLiveServerSessionHistoryController : public FServerSessionHistoryControllerBase
 {
 public:
diff --git a/Engine/Plugins/Developer/Concert/ConcertApp/MultiUserServer/Source/MultiUserServer/Private/Widgets/SessionTabs/ServerSessionHistoryControllerBase.h b/Engine/Plugins/Developer/Concert/ConcertApp/MultiUserServer/Source/MultiUserServer/Private/Widgets/SessionTabs/ServerSessionHistoryControllerBase.h
--- a/Engine/Plugins/Developer/Concert/ConcertApp/MultiUserServer/Source/MultiUserServer/Private/Widgets/SessionTabs/ServerSessionHistoryControllerBase.h
+++ b/Engine/Plugins/Developer/Concert/ConcertApp/MultiUserServer/Source/MultiUserServer/Private/Widgets/SessionTabs/ServerSessionHistoryControllerBase.h
@@ -12,7 +12,7 @@ class IConcertSyncServer;
 class IConcertServerSession;
 
 /**
- * Abstract base class providing most functionality.
+ * Abstract base class providing most functionality for managing SSessionHistory on the server.
  * Subclasses control which session database is used: archived or live.
  */
 class FServerSessionHistoryControllerBase : public FAbstractSessionHistoryController
diff --git a/Engine/Plugins/Developer/Concert/ConcertUI/ConcertSharedSlate/Source/ConcertSharedSlate/Private/Session/Activity/SConcertSessionActivities.cpp b/Engine/Plugins/Developer/Concert/ConcertUI/ConcertSharedSlate/Source/ConcertSharedSlate/Private/Session/Activity/SConcertSessionActivities.cpp
--- a/Engine/Plugins/Developer/Concert/ConcertUI/ConcertSharedSlate/Source/ConcertSharedSlate/Private/Session/Activity/SConcertSessionActivities.cpp
+++ b/Engine/Plugins/Developer/Concert/ConcertUI/ConcertSharedSlate/Source/ConcertSharedSlate/Private/Session/Activity/SConcertSessionActivities.cpp
@@ -246,7 +246,7 @@ void SConcertSessionActivities::Construct(const FArguments& InArgs)
 					SAssignNew(ActivityView, SListView<TSharedPtr<FConcertSessionActivity>>)
 					.ListItemsSource(&Activities)
 					.OnGenerateRow(this, &SConcertSessionActivities::OnGenerateActivityRowWidget)
-					.SelectionMode(ESelectionMode::Single)
+					.SelectionMode(InArgs._SelectionMode.Get(ESelectionMode::Single))
 					.AllowOverscroll(EAllowOverscroll::No)
 					.OnListViewScrolled(this, &SConcertSessionActivities::OnListViewScrolled)
 					.OnSelectionChanged(this, &SConcertSessionActivities::OnListViewSelectionChanged)
@@ -768,6 +768,11 @@ TSharedPtr<FConcertSessionActivity> SConcertSessionActivities::GetSelectedActivi
 	return SelectedItems.Num() ? SelectedItems[0] : nullptr;
 }
 
+TArray<TSharedPtr<FConcertSessionActivity>> SConcertSessionActivities::GetSelectedActivities() const
+{
+	return ActivityView->GetSelectedItems();
+}
+
 TSharedPtr<FConcertSessionActivity> SConcertSessionActivities::GetMostRecentActivity() const
 {
 	// NOTE: This function assumes that activities are sorted by ID. When used for recovery purpose,
diff --git a/Engine/Plugins/Developer/Concert/ConcertUI/ConcertSharedSlate/Source/ConcertSharedSlate/Private/Session/History/SEditableSessionHistory.cpp b/Engine/Plugins/Developer/Concert/ConcertUI/ConcertSharedSlate/Source/ConcertSharedSlate/Private/Session/History/SEditableSessionHistory.cpp
--- a/Engine/Plugins/Developer/Concert/ConcertUI/ConcertSharedSlate/Source/ConcertSharedSlate/Private/Session/History/SEditableSessionHistory.cpp
+++ b/Engine/Plugins/Developer/Concert/ConcertUI/ConcertSharedSlate/Source/ConcertSharedSlate/Private/Session/History/SEditableSessionHistory.cpp
@@ -4,6 +4,9 @@
 
 #include "EditorStyleSet.h"
 #include "SNegativeActionButton.h"
+
+#include "Algo/Transform.h"
+
 #include "Session/History/SSessionHistory.h"
 #include "Widgets/Layout/SBox.h"
 
@@ -16,9 +19,17 @@ void SEditableSessionHistory::Construct(const FArguments& InArgs)
 	DeleteActivityFunc = InArgs._DeleteActivity;
 	
 	SessionHistory = InArgs._MakeSessionHistory.Execute(
-			SSessionHistory::FArguments()
-			.OnMakeColumnOverlayWidget(this, &SEditableSessionHistory::MakeSummaryColumnDeleteButton)
-			);
+		SSessionHistory::FArguments()
+		.SelectionMode(ESelectionMode::Multi)
+		.SearchButtonArea()
+		[
+			SNew(SNegativeActionButton)
+			.OnClicked(this, &SEditableSessionHistory::OnClickDeleteActivityButton)
+			.ToolTipText(this, &SEditableSessionHistory::GetDeleteActivityToolTip)
+			.IsEnabled(this, &SEditableSessionHistory::IsDeleteButtonEnabled)
+			.Icon(FEditorStyle::GetBrush("Icons.Delete"))
+		]
+		);
 	
 	ChildSlot
 	[
@@ -26,56 +37,64 @@ void SEditableSessionHistory::Construct(const FArguments& InArgs)
 	];
 }
 
-TSharedPtr<SWidget> SEditableSessionHistory::MakeSummaryColumnDeleteButton(
-	TWeakPtr<SMultiColumnTableRow<TSharedPtr<FConcertSessionActivity>>> OwningRow,
-	TWeakPtr<FConcertSessionActivity> RowActivity,
-	const FName& ColumnId) const
+FReply SEditableSessionHistory::OnKeyDown(const FGeometry& MyGeometry, const FKeyEvent& InKeyEvent)
 {
-	if (!SessionHistory->IsLastColumn(ColumnId))
+	if (InKeyEvent.GetKey() == EKeys::Delete)
 	{
-		return nullptr;
+		const TSet<TSharedRef<FConcertSessionActivity>> SelectedActivities = SessionHistory->GetSelectedActivities();
+		if (CanDeleteActivityFunc.Execute(SelectedActivities).CanDelete())
+		{
+			DeleteActivityFunc.Execute(SelectedActivities);
+		}
+		return FReply::Handled();
 	}
-
-	return SNew(SBox)
-		.Padding(FMargin(1, 1))
-		.HAlign(HAlign_Right)
-		.VAlign(VAlign_Center)
-		[
-			SNew(SNegativeActionButton)
-			.OnClicked(this, &SEditableSessionHistory::OnClickDeleteActivityButton, RowActivity)
-			.Visibility_Lambda([OwningRow](){ return OwningRow.IsValid() && OwningRow.Pin()->IsHovered() ? EVisibility::Visible : EVisibility::Collapsed; })
-			.ToolTipText(TAttribute<FText>::CreateLambda([this, RowActivity](){ return GetDeleteActivityToolTip(RowActivity);}))
-			.IsEnabled(TAttribute<bool>::CreateLambda([this, RowActivity](){ return CanDeleteActivity(RowActivity);}))
-			.Icon(FEditorStyle::GetBrush("Icons.Delete"))
-		];
+	
+	return SCompoundWidget::OnKeyDown(MyGeometry, InKeyEvent);
 }
 
-FReply SEditableSessionHistory::OnClickDeleteActivityButton(TWeakPtr<FConcertSessionActivity> RowActivity) const
+FReply SEditableSessionHistory::OnClickDeleteActivityButton() const
 {
-	if (DeleteActivityFunc.IsBound() && CanDeleteActivity(RowActivity))
+	if (DeleteActivityFunc.IsBound() && CanDeleteActivityFunc.Execute(SessionHistory->GetSelectedActivities()).CanDelete())
 	{
-		DeleteActivityFunc.Execute(RowActivity.Pin().ToSharedRef());
+		DeleteActivityFunc.Execute(SessionHistory->GetSelectedActivities());
 	}
 	return FReply::Handled();
 }
 
-FText SEditableSessionHistory::GetDeleteActivityToolTip(TWeakPtr<FConcertSessionActivity> RowActivity) const
+FText SEditableSessionHistory::GetDeleteActivityToolTip() const
 {
-	return CanDeleteActivity(RowActivity) ?
-		LOCTEXT("DeleteActivityDescription", "Delete activity from history")
-		:
-		LOCTEXT("CannotDeleteActivityDescription", "Activity cannot be deleted from history");
-}
-
-bool SEditableSessionHistory::CanDeleteActivity(TWeakPtr<FConcertSessionActivity> RowActivity) const
-{
-	if (const TSharedPtr<FConcertSessionActivity> Pinned = RowActivity.Pin())
+	TSet<TSharedRef<FConcertSessionActivity>> SelectedActivities = SessionHistory->GetSelectedActivities();
+	if (SelectedActivities.Num() == 0)
 	{
-		return CanDeleteActivityFunc.Execute(Pinned.ToSharedRef());
+		return LOCTEXT("SelectActivityToolTip", "Select some activities to delete from below (multi-select using CTRL + Click).");
 	}
+	
+	const FCanDeleteActivitiesResult CanDeleteActivities = CanDeleteActivityFunc.Execute(SelectedActivities);
+	if (CanDeleteActivities.CanDelete())
+	{
+		SelectedActivities.Sort([](const TSharedRef<FConcertSessionActivity>& First, const TSharedRef<FConcertSessionActivity>& Second) { return First->Activity.ActivityId <= Second->Activity.ActivityId; });
+		
+		TStringBuilder<256> ActivityBuilder;
+		bool bNeedsComma = false;
+		for (const TSharedRef<FConcertSessionActivity>& Activity : SelectedActivities)
+		{
+			if (bNeedsComma)
+			{
+				ActivityBuilder << ", ";
+			}
+			ActivityBuilder << Activity->Activity.ActivityId;
 
-	checkNoEntry();
-	return false;
+			bNeedsComma = true;
+		}
+		return FText::Format(LOCTEXT("DeleteSelectedActivitiesToolTip", "Delete selected activities from history (IDs: {0})"), FText::FromString(ActivityBuilder.ToString()));
+	}
+	return FText::Format(LOCTEXT("CannotDeleteSelectedActivitiesToolTip", "Activity cannot be deleted: {0}"), CanDeleteActivities.DeletionReason.GetValue());
+}
+
+bool SEditableSessionHistory::IsDeleteButtonEnabled() const
+{
+	const TSet<TSharedRef<FConcertSessionActivity>> SelectedActivities = SessionHistory->GetSelectedActivities();
+	return SelectedActivities.Num() > 0 && CanDeleteActivityFunc.Execute(SelectedActivities).CanDelete();
 }
 
 #undef LOCTEXT_NAMESPACE
diff --git a/Engine/Plugins/Developer/Concert/ConcertUI/ConcertSharedSlate/Source/ConcertSharedSlate/Private/Session/History/SSessionHistory.cpp b/Engine/Plugins/Developer/Concert/ConcertUI/ConcertSharedSlate/Source/ConcertSharedSlate/Private/Session/History/SSessionHistory.cpp
--- a/Engine/Plugins/Developer/Concert/ConcertUI/ConcertSharedSlate/Source/ConcertSharedSlate/Private/Session/History/SSessionHistory.cpp
+++ b/Engine/Plugins/Developer/Concert/ConcertUI/ConcertSharedSlate/Source/ConcertSharedSlate/Private/Session/History/SSessionHistory.cpp
@@ -5,6 +5,8 @@
 #include "ConcertMessageData.h"
 #include "IConcertSession.h"
 
+#include "Algo/Transform.h"
+
 #include "Session/Activity/PredefinedActivityColumns.h"
 #include "Session/Activity/SConcertSessionActivities.h"
 
@@ -46,10 +48,14 @@ void SSessionHistory::Construct(const FArguments& InArgs)
 {
 	using namespace UE::ConcertSharedSlate;
 	
-	PackageNameFilter = InArgs._PackageFilter;
+	AllowActivityFunc = FAllowActivity::CreateLambda([this, PackageNameFilter = InArgs._PackageFilter, FilterFunc = InArgs._AllowActivity](const FConcertSyncActivity& Activity, const TStructOnScope<FConcertSyncActivitySummary>& Summary)
+	{
+		return ConcertSessionHistoryUI::PackageNamePassesFilter(PackageNameFilter, Summary)
+			&& (!FilterFunc.IsBound() || FilterFunc.Execute(Activity, Summary));
+	});
 
 	ActivityMap.Reserve(MaximumNumberOfActivities);
-	ActivityListViewOptions = MakeShared<FConcertSessionActivitiesOptions>();
+	ActivityListViewOptions = InArgs._ViewOptions.Get() ? InArgs._ViewOptions.Get() : MakeShared<FConcertSessionActivitiesOptions>();
 
 	SAssignNew(ActivityListView, SConcertSessionActivities)
 		.OnGetPackageEvent(InArgs._GetPackageEvent)
@@ -58,19 +64,16 @@ void SSessionHistory::Construct(const FArguments& InArgs)
 		.OnMapActivityToClient(this, &SSessionHistory::GetClientInfo)
 		.HighlightText(this, &SSessionHistory::HighlightSearchedText)
 		.TimeFormat(ActivityListViewOptions.Get(), &FConcertSessionActivitiesOptions::GetTimeFormat)
-		.Columns({
-			ActivityColumn::AvatarColor(),
-			ActivityColumn::ClientName(),
-			ActivityColumn::Operation()
-		})
+		.Columns(InArgs._Columns)
 		.ConnectionActivitiesVisibility(ActivityListViewOptions.Get(), &FConcertSessionActivitiesOptions::GetConnectionActivitiesVisibility)
 		.LockActivitiesVisibility(ActivityListViewOptions.Get(), &FConcertSessionActivitiesOptions::GetLockActivitiesVisibility)
 		.PackageActivitiesVisibility(ActivityListViewOptions.Get(), &FConcertSessionActivitiesOptions::GetPackageActivitiesVisibility)
 		.TransactionActivitiesVisibility(ActivityListViewOptions.Get(), &FConcertSessionActivitiesOptions::GetTransactionActivitiesVisibility)
-		.DetailsAreaVisibility(EVisibility::Visible)
+		.DetailsAreaVisibility(InArgs._DetailsAreaVisibility)
 		.IsAutoScrollEnabled(true)
 		.ColumnVisibilitySnapshot(InArgs._ColumnVisibilitySnapshot)
-		.SaveColumnVisibilitySnapshot(InArgs._SaveColumnVisibilitySnapshot);
+		.SaveColumnVisibilitySnapshot(InArgs._SaveColumnVisibilitySnapshot)
+		.SelectionMode(InArgs._SelectionMode);
 
 	ChildSlot
 	[
@@ -79,11 +82,23 @@ void SSessionHistory::Construct(const FArguments& InArgs)
 		.AutoHeight()
 		.Padding(1, 1)
 		[
-			SAssignNew(SearchBox, SSearchBox)
-			.HintText(LOCTEXT("SearchHint", "Search..."))
-			.OnTextChanged(this, &SSessionHistory::OnSearchTextChanged)
-			.OnTextCommitted(this, &SSessionHistory::OnSearchTextCommitted)
-			.DelayChangeNotificationsWhileTyping(true)
+			SNew(SHorizontalBox)
+
+			+SHorizontalBox::Slot()
+			.AutoWidth()
+			[
+				InArgs._SearchButtonArea.Widget
+			]
+
+			+SHorizontalBox::Slot()
+			.FillWidth(1.f)
+			[
+				SAssignNew(SearchBox, SSearchBox)
+				.HintText(LOCTEXT("SearchHint", "Search..."))
+				.OnTextChanged(this, &SSessionHistory::OnSearchTextChanged)
+				.OnTextCommitted(this, &SSessionHistory::OnSearchTextCommitted)
+				.DelayChangeNotificationsWhileTyping(true)
+			]
 		]
 
 		+SVerticalBox::Slot()
@@ -122,7 +137,7 @@ void SSessionHistory::ReloadActivities(TMap<FGuid, FConcertClientInfo> InEndpoin
 
 	for (FConcertSessionActivity& FetchedActivity : InFetchedActivities)
 	{
-		if (ConcertSessionHistoryUI::PackageNamePassesFilter(PackageNameFilter, FetchedActivity.ActivitySummary))
+		if (AllowActivityFunc.Execute(FetchedActivity.Activity, FetchedActivity.ActivitySummary))
 		{
 			TSharedRef<FConcertSessionActivity> NewActivity = MakeShared<FConcertSessionActivity>(MoveTemp(FetchedActivity));
 			ActivityMap.Add(NewActivity->Activity.ActivityId, NewActivity);
@@ -136,7 +151,7 @@ void SSessionHistory::HandleActivityAddedOrUpdated(const FConcertClientInfo& InC
 	TStructOnScope<FConcertSyncActivitySummary> ActivitySummary;
 	ActivitySummary.InitializeFromChecked(InActivitySummary);
 
-	if (ConcertSessionHistoryUI::PackageNamePassesFilter(PackageNameFilter, ActivitySummary))
+	if (AllowActivityFunc.Execute(InActivity, ActivitySummary))
 	{
 		EndpointClientInfoMap.Add(InActivity.EndpointId, InClientInfo);
 
@@ -163,6 +178,16 @@ void SSessionHistory::OnColumnVisibilitySettingsChanged(const FColumnVisibilityS
 	ActivityListView->OnColumnVisibilitySettingsChanged(ColumnSnapshot);
 }
 
+TSet<TSharedRef<FConcertSessionActivity>> SSessionHistory::GetSelectedActivities() const
+{
+	TSet<TSharedRef<FConcertSessionActivity>> Activities;
+	Algo::Transform(ActivityListView->GetSelectedActivities(), Activities, [&](TSharedPtr<FConcertSessionActivity> Activity)
+	{
+		return Activity.ToSharedRef();
+	});
+	return Activities;
+}
+
 void SSessionHistory::OnSearchTextChanged(const FText& InSearchText)
 {
 	SearchedText = InSearchText;
diff --git a/Engine/Plugins/Developer/Concert/ConcertUI/ConcertSharedSlate/Source/ConcertSharedSlate/Public/Session/Activity/SConcertSessionActivities.h b/Engine/Plugins/Developer/Concert/ConcertUI/ConcertSharedSlate/Source/ConcertSharedSlate/Public/Session/Activity/SConcertSessionActivities.h
--- a/Engine/Plugins/Developer/Concert/ConcertUI/ConcertSharedSlate/Source/ConcertSharedSlate/Public/Session/Activity/SConcertSessionActivities.h
+++ b/Engine/Plugins/Developer/Concert/ConcertUI/ConcertSharedSlate/Source/ConcertSharedSlate/Public/Session/Activity/SConcertSessionActivities.h
@@ -139,6 +139,9 @@ public:
 		/** Called whenever the column visibility changes and should be saved */
 		SLATE_EVENT(UE::ConcertSharedSlate::FSaveColumnVisibilitySnapshot, SaveColumnVisibilitySnapshot)
 	
+
+		/** How the activities may be selected */
+		SLATE_ARGUMENT(TOptional<ESelectionMode::Type>, SelectionMode)
 	SLATE_END_ARGS();
 	
 	/**
@@ -152,6 +155,8 @@ public:
 
 	/** Returns the activity selected or null if none is selected. */
 	TSharedPtr<FConcertSessionActivity> GetSelectedActivity() const;
+	/** Returns all selected activities */
+	TArray<TSharedPtr<FConcertSessionActivity>> GetSelectedActivities() const;
 
 	/** Returns the total number of activities currently stored (no filter applied). */
 	int32 GetTotalActivityNum() const { return AllActivities.Num(); }
diff --git a/Engine/Plugins/Developer/Concert/ConcertUI/ConcertSharedSlate/Source/ConcertSharedSlate/Public/Session/History/SEditableSessionHistory.h b/Engine/Plugins/Developer/Concert/ConcertUI/ConcertSharedSlate/Source/ConcertSharedSlate/Public/Session/History/SEditableSessionHistory.h
--- a/Engine/Plugins/Developer/Concert/ConcertUI/ConcertSharedSlate/Source/ConcertSharedSlate/Public/Session/History/SEditableSessionHistory.h
+++ b/Engine/Plugins/Developer/Concert/ConcertUI/ConcertSharedSlate/Source/ConcertSharedSlate/Public/Session/History/SEditableSessionHistory.h
@@ -10,35 +10,51 @@
 class SSessionHistory;
 struct FConcertSessionActivity;
 
+struct FCanDeleteActivitiesResult
+{
+	TOptional<FText> DeletionReason;
+
+	static FCanDeleteActivitiesResult Yes() { return FCanDeleteActivitiesResult(); }
+	static FCanDeleteActivitiesResult No(FText Reason) { return FCanDeleteActivitiesResult(MoveTemp(Reason)); }
+
+	bool CanDelete() const { return !DeletionReason.IsSet(); }
+
+	explicit FCanDeleteActivitiesResult(TOptional<FText> DeletionReason = {})
+		: DeletionReason(DeletionReason)
+	{}
+};
+
 /** Allows activities in the session history to be deleted. */
 class CONCERTSHAREDSLATE_API SEditableSessionHistory : public SCompoundWidget
 {
 public:
 
 	DECLARE_DELEGATE_RetVal_OneParam(TSharedRef<SSessionHistory>, FMakeSessionHistory, SSessionHistory::FArguments)
-	DECLARE_DELEGATE_RetVal_OneParam(bool, FCanDeleteActivity, const TSharedRef<FConcertSessionActivity>&)
-	DECLARE_DELEGATE_OneParam(FRequestDeleteActivity, const TSharedRef<FConcertSessionActivity>&)
+	DECLARE_DELEGATE_RetVal_OneParam(FCanDeleteActivitiesResult, FCanDeleteActivities, const TSet<TSharedRef<FConcertSessionActivity>>& /*Activities*/)
+	DECLARE_DELEGATE_OneParam(FRequestDeleteActivities, const TSet<TSharedRef<FConcertSessionActivity>>&)
 
 	SLATE_BEGIN_ARGS(SEditableSessionHistory)
 	{}
 		SLATE_EVENT(FMakeSessionHistory, MakeSessionHistory)
-		SLATE_EVENT(FCanDeleteActivity, CanDeleteActivity)
-		SLATE_EVENT(FRequestDeleteActivity, DeleteActivity)
+		SLATE_EVENT(FCanDeleteActivities, CanDeleteActivity)
+		SLATE_EVENT(FRequestDeleteActivities, DeleteActivity)
 		SLATE_NAMED_SLOT(FArguments, StatusBar)
 	SLATE_END_ARGS()
 
 	void Construct(const FArguments& InArgs);
-
+	
+	//~ Begin SWidget Interface
+	virtual FReply OnKeyDown(const FGeometry& MyGeometry, const FKeyEvent& InKeyEvent) override;
+	//~ End SWidget Interface
+	
 private:
 
 	TSharedPtr<SSessionHistory> SessionHistory;
 
-	FCanDeleteActivity CanDeleteActivityFunc;
-	FRequestDeleteActivity DeleteActivityFunc;
+	FCanDeleteActivities CanDeleteActivityFunc;
+	FRequestDeleteActivities DeleteActivityFunc;
 
-	/** Creates a delete button that is overlayed over the summary column*/
-	TSharedPtr<SWidget> MakeSummaryColumnDeleteButton(TWeakPtr<SMultiColumnTableRow<TSharedPtr<FConcertSessionActivity>>> OwningRow, TWeakPtr<FConcertSessionActivity> RowActivity, const FName& ColumnId) const;
-	FReply OnClickDeleteActivityButton(TWeakPtr<FConcertSessionActivity> RowActivity) const;
-	FText GetDeleteActivityToolTip(TWeakPtr<FConcertSessionActivity> RowActivity) const;
-	bool CanDeleteActivity(TWeakPtr<FConcertSessionActivity> RowActivity) const;
+	FReply OnClickDeleteActivityButton() const;
+	FText GetDeleteActivityToolTip() const;
+	bool IsDeleteButtonEnabled() const;
 };
diff --git a/Engine/Plugins/Developer/Concert/ConcertUI/ConcertSharedSlate/Source/ConcertSharedSlate/Public/Session/History/SSessionHistory.h b/Engine/Plugins/Developer/Concert/ConcertUI/ConcertSharedSlate/Source/ConcertSharedSlate/Public/Session/History/SSessionHistory.h
--- a/Engine/Plugins/Developer/Concert/ConcertUI/ConcertSharedSlate/Source/ConcertSharedSlate/Public/Session/History/SSessionHistory.h
+++ b/Engine/Plugins/Developer/Concert/ConcertUI/ConcertSharedSlate/Source/ConcertSharedSlate/Public/Session/History/SSessionHistory.h
@@ -3,6 +3,7 @@
 #pragma once
 
 #include "CoreMinimal.h"
+#include "Session/Activity/PredefinedActivityColumns.h"
 #include "Session/Activity/SConcertSessionActivities.h"
 #include "Widgets/DeclarativeSyntaxSupport.h"
 #include "Widgets/SCompoundWidget.h"
@@ -10,33 +11,63 @@
 class SSearchBox;
 class FConcertSessionActivitiesOptions;
 
+/** Wraps a session table list with a search box and a status bar for visibility settings (show connection activities, etc.). */
 class CONCERTSHAREDSLATE_API SSessionHistory : public SCompoundWidget
 {
 public:
 
+	DECLARE_DELEGATE_RetVal_TwoParams(bool, FAllowActivity, const FConcertSyncActivity&, const TStructOnScope<FConcertSyncActivitySummary>&)
+
 	/** Maximum number of activities displayed on screen. */ 
 	static constexpr int64 MaximumNumberOfActivities = 1000;
 	
-	SLATE_BEGIN_ARGS(SSessionHistory) {}
-		SLATE_ARGUMENT(FName, PackageFilter)
+	SLATE_BEGIN_ARGS(SSessionHistory)
+		: _Columns({
+			UE::ConcertSharedSlate::ActivityColumn::AvatarColor(),
+			UE::ConcertSharedSlate::ActivityColumn::ClientName(),
+			UE::ConcertSharedSlate::ActivityColumn::Operation()
+		})
+		, _DetailsAreaVisibility(EVisibility::Visible)
+		, _SelectionMode(ESelectionMode::Single)
+	{}
+	
 		SLATE_EVENT(SConcertSessionActivities::FGetPackageEvent, GetPackageEvent)
 		SLATE_EVENT(SConcertSessionActivities::FGetTransactionEvent, GetTransactionEvent)
-		/** If bound, invoked when generating a row to add an overlay to a column. */
+	
+		/** Optional name for only displaying activities that affect certain packages */
+		SLATE_ARGUMENT(FName, PackageFilter)
+		/** Optional filter for filtering activities out */
+		SLATE_EVENT(FAllowActivity, AllowActivity)
+	
+		/** Optional. The columns to show for each activity; they are shown after date time. Defaults to avatar color, client name, and operation. */
+		SLATE_ARGUMENT(TArray<FActivityColumn>, Columns)
+		/** Optional. If bound, invoked when generating a row to add an overlay to a column. */
 		SLATE_EVENT(SConcertSessionActivities::FMakeColumnOverlayWidgetFunc, OnMakeColumnOverlayWidget)
 	
 		/** Optional snapshot to restore column visibilities with */
 		SLATE_ARGUMENT(FColumnVisibilitySnapshot, ColumnVisibilitySnapshot)
 		/** Called whenever the column visibility changes and should be saved */
 		SLATE_EVENT(UE::ConcertSharedSlate::FSaveColumnVisibilitySnapshot, SaveColumnVisibilitySnapshot)
+
+		/** Optional. You can override the view options with this. */
+		SLATE_ATTRIBUTE(TSharedPtr<FConcertSessionActivitiesOptions>, ViewOptions)
+	
+		/** Optional. Show/hide the details area widget. Default to visible. */
+		SLATE_ARGUMENT(EVisibility, DetailsAreaVisibility)
+		
+		/** Optional. How the activities may be selected. Default to Single. */
+		SLATE_ARGUMENT(ESelectionMode::Type, SelectionMode)
+
+		/** Optional. An area to the left of the search bar intended for adding buttons to. */
+		SLATE_NAMED_SLOT(FArguments, SearchButtonArea)
+	
 	SLATE_END_ARGS()
 
 	/**
-	* Constructs the Session History widget.
-	*
-	* @param InArgs The Slate argument list.
-	* @param ConstructUnderMajorTab The major tab which will contain the Session History widget.
-	* @param InConcertSyncClient Pointer on the concert sync client.
-	*/
+	 * Constructs the Session History widget.
+	 *
+	 * @param InArgs The Slate argument list.
+	 */
 	void Construct(const FArguments& InArgs);
 
 	/** Fetches activities from the server and updates the list view. */
@@ -46,6 +77,8 @@ public:
 	void HandleActivityAddedOrUpdated(const FConcertClientInfo& InClientInfo, const FConcertSyncActivity& InActivity, const FStructOnScope& InActivitySummary);
 
 	bool IsLastColumn(FName ColumnId) const { return ActivityListView->IsLastColumn(ColumnId); }
+	
+	TSet<TSharedRef<FConcertSessionActivity>> GetSelectedActivities() const;
 
 	void OnColumnVisibilitySettingsChanged(const FColumnVisibilitySnapshot& ColumnSnapshot);
 
@@ -80,8 +113,8 @@ private:
 	/** The searched text to highlight. */
 	FText SearchedText;
 
-	/** Used to limit activities to a given package only. */
-	FName PackageNameFilter;
+	/** Used to limit activities. */
+	FAllowActivity AllowActivityFunc;
 
 	TOptional<FConcertClientInfo> GetClientInfo(FGuid Guid) const;
 };
\ No newline at end of file
