Allow deleting multiple sessions at once

diff --git a/Engine/Plugins/Developer/Concert/ConcertApp/MultiUserClient/Source/MultiUserClient/Private/Widgets/SConcertBrowser.cpp b/Engine/Plugins/Developer/Concert/ConcertApp/MultiUserClient/Source/MultiUserClient/Private/Widgets/SConcertBrowser.cpp
--- a/Engine/Plugins/Developer/Concert/ConcertApp/MultiUserClient/Source/MultiUserClient/Private/Widgets/SConcertBrowser.cpp
+++ b/Engine/Plugins/Developer/Concert/ConcertApp/MultiUserClient/Source/MultiUserClient/Private/Widgets/SConcertBrowser.cpp
@@ -14,6 +14,7 @@
 #include "MultiUserClientUtils.h"
 #include "SActiveSession.h"
 #include "SConcertSessionRecovery.h"
+
 #include "Session/Browser/ConcertBrowserUtils.h"
 #include "Session/Browser/ConcertSessionItem.h"
 #include "Session/Browser/IConcertSessionBrowserController.h"
@@ -142,8 +143,7 @@ public:
 	virtual void RenameArchivedSession(const FGuid& ServerAdminEndpointId, const FGuid& SessionId, const FString& NewName) override;
 	virtual bool CanRenameActiveSession(const FGuid& ServerAdminEndpointId, const FGuid& SessionId) const override;
 	virtual bool CanRenameArchivedSession(const FGuid& ServerAdminEndpointId, const FGuid& SessionId) const override;
-	virtual void DeleteActiveSession(const FGuid& ServerAdminEndpointId, const FGuid& SessionId) override;
-	virtual void DeleteArchivedSession(const FGuid& ServerAdminEndpointId, const FGuid& SessionId) override;
+	virtual void DeleteSessions(const FGuid& ServerAdminEndpointId, const TArray<FGuid>& SessionIds) override;
 	virtual bool CanDeleteActiveSession(const FGuid& ServerAdminEndpointId, const FGuid& SessionId) const override;
 	virtual bool CanDeleteArchivedSession(const FGuid& ServerAdminEndpointId, const FGuid& SessionId) const override;
 	//~ End IConcertSessionBrowserController Interface
@@ -231,7 +231,6 @@ private:
 	TSet<FString> IgnoredServers; // List of ignored servers (Useful for testing/debugging)
 };
 
-
 FConcertClientSessionBrowserController::FConcertClientSessionBrowserController(IConcertClientPtr InConcertClient)
 {
 	check(InConcertClient.IsValid()); // Don't expect this class to be instantiated if the concert client is not available.
@@ -482,18 +481,21 @@ bool FConcertClientSessionBrowserController::CanRenameArchivedSession(const FGui
 	return CanDeleteArchivedSession(ServerAdminEndpointId, SessionId); // Rename requires the same permission than delete.
 }
 
-void FConcertClientSessionBrowserController::DeleteActiveSession(const FGuid& ServerAdminEndpointId, const FGuid& SessionId)
+void FConcertClientSessionBrowserController::DeleteSessions(const FGuid& ServerAdminEndpointId, const TArray<FGuid>& SessionIds)
 {
-	// On success, an session is deleted and TickServersAndSessionsDiscovery() will eventually notice it.
-	// On failure: An async notification banner will be displayer to the user.
-	ConcertClient->DeleteSession(ServerAdminEndpointId, SessionId);
-}
-
-void FConcertClientSessionBrowserController::DeleteArchivedSession(const FGuid& ServerAdminEndpointId, const FGuid& SessionId)
-{
-	// On success, an archive is deleted and TickServersAndSessionsDiscovery() will eventually notice it.
-	// On failure: An async notification banner will be displayer to the user.
-	ConcertClient->DeleteSession(ServerAdminEndpointId, SessionId);
+	// The difference between the two functions is how error notifications are displayed...
+	if (SessionIds.Num() == 1)
+	{
+		// ... shows a specific notification
+		ConcertClient->DeleteSession(ServerAdminEndpointId, SessionIds[0]);
+	}
+	else
+	{
+		TSet<FGuid> SessionIdSet;
+		Algo::Transform(SessionIds, SessionIdSet, [](const FGuid& SessionId) { return SessionId; });
+		// ... shows compressed notification(s)
+		ConcertClient->BatchDeleteSessions(ServerAdminEndpointId, { SessionIdSet, EBatchSessionDeletionFlags::SkipForbiddenSessions });
+	}
 }
 
 bool FConcertClientSessionBrowserController::CanDeleteActiveSession(const FGuid& ServerAdminEndpointId, const FGuid& SessionId) const
@@ -1038,7 +1040,7 @@ private:
 	// Manipulates the sessions view (the array and the UI).
 	void OnSessionSelectionChanged(const TSharedPtr<FConcertSessionItem>& SessionItem);
 	void OnSessionDoubleClicked(const TSharedPtr<FConcertSessionItem>& SessionItem);
-	bool ConfirmDeleteSessionWithDialog(const TSharedPtr<FConcertSessionItem>& SessionItem);
+	bool ConfirmDeleteSessionWithDialog(const TArray<TSharedPtr<FConcertSessionItem>>& SessionItems) const;
 
 	// Update server/session/clients lists.
 	EActiveTimerReturnType TickDiscovery(double InCurrentTime, float InDeltaTime);
@@ -1112,7 +1114,7 @@ void SConcertClientSessionBrowser::Construct(const FArguments& InArgs, IConcertC
 	// Displayed when discovering session or if no session is available.
 	SessionDiscoveryPanel = SNew(SConcertDiscovery)
 		.Text_Lambda(GetNoSessionText)
-		.Visibility_Lambda([this]() { return Controller->GetServers().Num() > 0 && SessionBrowser->GetSessions().Num() == 0 && !Controller->IsCreatingSession() ? EVisibility::Visible : EVisibility::Hidden; })
+		.Visibility_Lambda([this]() { return Controller->GetServers().Num() > 0 && !SessionBrowser->HasAnySessions() && !Controller->IsCreatingSession() ? EVisibility::Visible : EVisibility::Hidden; })
 		.ThrobberVisibility_Lambda([this]() { return !Controller->HasReceivedInitialSessionList() ? EVisibility::Visible : EVisibility::Collapsed; })
 		.ButtonVisibility_Lambda([this]() { return Controller->HasReceivedInitialSessionList() && Controller->GetActiveSessions().Num() == 0 && Controller->GetArchivedSessions().Num() == 0 ? EVisibility::Visible : EVisibility::Collapsed; })
 		.ButtonStyle(FAppStyle::Get(), "FlatButton.Default")
@@ -1173,11 +1175,10 @@ TSharedRef<SWidget> SConcertClientSessionBrowser::MakeBrowserContent(TSharedPtr<
 					[
 						MakeUserAndSettings()
 					]
-					.OnSessionClicked_Raw(this, &SConcertClientSessionBrowser::OnSessionSelectionChanged)
-					.OnLiveSessionDoubleClicked_Raw(this, &SConcertClientSessionBrowser::OnSessionDoubleClicked)
-					.OnRequestedDeleteSession_Lambda([this](auto) { UpdateDiscovery(); /* Don't wait up to 1s, kick discovery right now */ })
-					.CanDeleteArchivedSession(this, &SConcertClientSessionBrowser::ConfirmDeleteSessionWithDialog)
-					.CanDeleteActiveSession(this, &SConcertClientSessionBrowser::ConfirmDeleteSessionWithDialog)
+					.OnSessionClicked(this, &SConcertClientSessionBrowser::OnSessionSelectionChanged)
+					.OnLiveSessionDoubleClicked(this, &SConcertClientSessionBrowser::OnSessionDoubleClicked)
+					.PostRequestedDeleteSession_Lambda([this](auto) { UpdateDiscovery(); /* Don't wait up to 1s, kick discovery right now */ })
+					.AskUserToDeleteSessions(this, &SConcertClientSessionBrowser::ConfirmDeleteSessionWithDialog)
 			]
 
 			// Session details.
@@ -1710,11 +1711,27 @@ void SConcertClientSessionBrowser::OnSessionDoubleClicked(const TSharedPtr<FConc
 	}
 }
 
-bool SConcertClientSessionBrowser::ConfirmDeleteSessionWithDialog(const TSharedPtr<FConcertSessionItem>& SessionItem)
+bool SConcertClientSessionBrowser::ConfirmDeleteSessionWithDialog(const TArray<TSharedPtr<FConcertSessionItem>>& SessionItems) const
 {
-	const FText SessionNameInText = FText::FromString(SessionItem->SessionName);
-	const FText SeverNameInText = FText::FromString(SessionItem->ServerName);
-	const FText ConfirmationMessage = FText::Format(LOCTEXT("DeleteSessionConfirmationMessage", "Do you really want to delete the session \"{0}\" from the server \"{1}\"?"), SessionNameInText, SeverNameInText);
+	TSet<FString> SessionNames;
+	TSet<FString> UniqueServers;
+	
+	Algo::Transform(SessionItems, SessionNames, [](const TSharedPtr<FConcertSessionItem>& Item)
+	{
+		return Item->SessionName;
+	});
+	Algo::Transform(SessionItems, UniqueServers, [](const TSharedPtr<FConcertSessionItem>& Item)
+	{
+		return Item->ServerName;
+	});
+	
+	const FText ConfirmationMessage = FText::Format(
+		LOCTEXT("DeleteSessionConfirmationMessage", "Do you really want to delete {0} {0}|plural(one=session,other=sessions) from {1} {1}|plural(one=server,other=servers)?\n\nThe {0}|plural(one=session,other=sessions) {2} will be deleted from the {1}|plural(one=server,other=servers) {3}."),
+		SessionItems.Num(),
+		UniqueServers.Num(),
+		FText::FromString(FString::Join(SessionNames, TEXT(", "))),
+		FText::FromString(FString::Join(UniqueServers, TEXT(", ")))
+		);
 	const FText ConfirmationTitle = LOCTEXT("DeleteSessionConfirmationTitle", "Delete Session Confirmation");
 	return FMessageDialog::Open(EAppMsgType::YesNo, ConfirmationMessage, &ConfirmationTitle) == EAppReturnType::Yes;
 }
diff --git a/Engine/Plugins/Developer/Concert/ConcertApp/MultiUserServer/Source/MultiUserServer/Private/Widgets/Browser/ConcertServerSessionBrowserController.cpp b/Engine/Plugins/Developer/Concert/ConcertApp/MultiUserServer/Source/MultiUserServer/Private/Widgets/Browser/ConcertServerSessionBrowserController.cpp
--- a/Engine/Plugins/Developer/Concert/ConcertApp/MultiUserServer/Source/MultiUserServer/Private/Widgets/Browser/ConcertServerSessionBrowserController.cpp
+++ b/Engine/Plugins/Developer/Concert/ConcertApp/MultiUserServer/Source/MultiUserServer/Private/Widgets/Browser/ConcertServerSessionBrowserController.cpp
@@ -111,22 +111,37 @@ void FConcertServerSessionBrowserController::CreateSession(const FGuid& ServerAd
 	VersionInfo.Initialize();
 	SessionInfo.VersionInfos.Emplace(VersionInfo);
 
+	FAsyncTaskNotification Notification(
+		MakeAsyncNotification(FText::Format(LOCTEXT("CreateSessionFmt.InProgress", "Created Session '{0}'"), FText::FromString(SessionName)))
+		);
 	FText FailureReason = FText::GetEmpty();
 	const bool bSuccess = ServerInstance->GetConcertServer()->CreateSession(SessionInfo, FailureReason) != nullptr;
 
-	NotifyUserOfFinishedSessionAction(bSuccess,
-		bSuccess ? FText::Format(LOCTEXT("CreatedSessionFmt", "Created Session '{0}'"), FText::FromString(SessionName)) : FText::Format(LOCTEXT("FailedToCreateSessionFmt", "Failed to create Session '{0}'"), FText::FromString(SessionName)),
+	CloseAsyncNotification(
+		Notification,
+		bSuccess,
+		bSuccess
+			? FText::Format(LOCTEXT("CreateSessionFmt.Success", "Created Session '{0}'"), FText::FromString(SessionName))
+			: FText::Format(LOCTEXT("CreateSessionFmt.Failure", "Failed to create Session '{0}'"), FText::FromString(SessionName)),
 		FailureReason
 		);
 }
 
 void FConcertServerSessionBrowserController::ArchiveSession(const FGuid& ServerAdminEndpointId, const FGuid& SessionId, const FString& ArchiveName, const FConcertSessionFilter& SessionFilter)
 {
+	FAsyncTaskNotification Notification(
+		MakeAsyncNotification(FText::Format(LOCTEXT("ArchivedSessionFmt.InProgress", "Archiving new Session '{0}'"), FText::FromString(ArchiveName)))
+		);
+	
 	FText FailureReason = FText::GetEmpty();
 	const bool bSuccess = ServerInstance->GetConcertServer()->ArchiveSession(SessionId, ArchiveName, SessionFilter, FailureReason).IsValid();
 	
-	NotifyUserOfFinishedSessionAction(bSuccess,
-		bSuccess ? FText::Format(LOCTEXT("ArchivedSessionFmt", "Archived Session '{0}'"),  FText::FromString(ArchiveName)) : FText::Format(LOCTEXT("FailedToArchivedSessionFmt", "Failed to archive Session '{0}'"), FText::FromString(ArchiveName)),
+	CloseAsyncNotification(
+		Notification,
+		bSuccess,
+		bSuccess
+			? FText::Format(LOCTEXT("ArchiveSessionFmt.Success", "Archived Session '{0}'"), FText::FromString(ArchiveName))
+			: FText::Format(LOCTEXT("ArchiveSessionFmt.Failure", "Failed to archive Session '{0}'"), FText::FromString(ArchiveName)),
 		FailureReason
 		);
 }
@@ -135,13 +150,94 @@ void FConcertServerSessionBrowserController::RestoreSession(const FGuid& ServerA
 {
 	if (TOptional<FConcertSessionInfo> SessionInfo = GetArchivedSessionInfo(ServerAdminEndpointId, SessionId))
 	{
+		FAsyncTaskNotification Notification(
+			MakeAsyncNotification(FText::Format(LOCTEXT("RestoreSessionFmt.InProgress", "Restoring Session '{0}'"), FText::FromString(RestoredName)))
+			);
+		
 		FText FailureReason = FText::GetEmpty();
 		SessionInfo->SessionName = RestoredName;
 		const bool bSuccess = ServerInstance->GetConcertServer()->RestoreSession(SessionId, *SessionInfo, SessionFilter, FailureReason).IsValid();
-		NotifyUserOfFinishedSessionAction(bSuccess,
-			bSuccess ? FText::Format(LOCTEXT("RestoreSessionFmt", "Restored Session '{0}'"),  FText::FromString(RestoredName)) : FText::Format(LOCTEXT("FailedToRestoreSessionFmt", "Failed to restore Session '{0}'"), FText::FromString(RestoredName)),
+		
+		CloseAsyncNotification(
+			Notification,
+			bSuccess,
+			bSuccess
+				? FText::Format(LOCTEXT("RestoreSessionFmt.Success", "Restored Session '{0}'"),  FText::FromString(RestoredName))
+				: FText::Format(LOCTEXT("RestoreSessionFmt.Failure", "Failed to restore Session '{0}'"), FText::FromString(RestoredName)),
 			FailureReason
+		);
+	}
+}
+
+void FConcertServerSessionBrowserController::DeleteSessions(const FGuid& ServerAdminEndpointId, const TArray<FGuid>& SessionIds)
+{
+	// We want a special message to be displayed when only deleting single items
+	if (SessionIds.Num() == 1)
+	{
+		const FString SessionName = GetSessionName(ServerAdminEndpointId, SessionIds[0]).Get(SessionIds[0].ToString());
+		FAsyncTaskNotification Notification(
+			MakeAsyncNotification(FText::Format(LOCTEXT("DeleteSessionFmt.InProgress", "Deleting session '{0}'"), FText::FromString(SessionName)))
 			);
+		
+		FText FailureReason;
+		const bool bSuccess = DeleteSession(ServerAdminEndpointId, SessionIds[0], FailureReason);
+		
+		CloseAsyncNotification(
+			Notification,
+			bSuccess,
+			bSuccess
+				? FText::Format(LOCTEXT("DeleteSessionFmt.Success", "Deleted Session '{0}'"), FText::FromString(SessionName))
+				: FText::Format(LOCTEXT("DeleteSessionFmt.Failure", "Failed to archive Session '{0}'"), FText::FromString(SessionIds[0].ToString())),
+			FailureReason
+		);
+		return;
+	}
+
+	FAsyncTaskNotification Notification(
+		MakeAsyncNotification(FText::Format(LOCTEXT("DeleteSessionsFmt.InProgress", "Deleting {0} sessions"), SessionIds.Num()))
+		);
+	
+	using FSessionName = FString;
+	using FReasonText = FText;
+	
+	TArray<TPair<FSessionName, FReasonText>> Failures;
+	int32 NumberSuccessful = 0;
+	for (const FGuid& SessionId : SessionIds)
+	{
+		FText FailureReason;
+		if (DeleteSession(ServerAdminEndpointId, SessionId, FailureReason))
+		{
+			++NumberSuccessful;
+		}
+		else 
+		{
+			const TOptional<FString> SessionName = GetSessionName(ServerAdminEndpointId, SessionId);
+			Failures.Add({ SessionName.Get(SessionId.ToString()), FailureReason });
+		}
+	}
+	
+	if (Failures.Num() == 0)
+	{
+		CloseAsyncNotification(
+			Notification,
+			true,
+			FText::Format(LOCTEXT("DeleteSessionsFmt.Success", "Deleted {0} Sessions"), NumberSuccessful)
+		);
+	}
+	else
+	{
+		// Instead of a notification for each failure, show one big one so the user does not have to manually close each (cuz that's annoying...)
+		TStringBuilder<512> SubText;
+		for (const TPair<FSessionName, FReasonText>& Failure : Failures)
+		{
+			SubText << Failure.Key << TEXT(": ") << Failure.Value.ToString() << TEXT("\n");
+		}
+		CloseAsyncNotification(
+			Notification,
+			false,
+			FText::Format(LOCTEXT("DeletedSessionsFmt.Failure", "Deleted {0} of {1} Sessions."), NumberSuccessful, SessionIds.Num()),
+			FText::FromString(SubText.ToString())
+		);
 	}
 }
 
@@ -193,36 +289,54 @@ void FConcertServerSessionBrowserController::RenameArchivedSessionInternal(const
 
 void FConcertServerSessionBrowserController::RenameSessionInternal(const FGuid& SessionId, const FString& NewName, const FConcertSessionInfo& SessionInfo)
 {
+	FAsyncTaskNotification Notification(
+		MakeAsyncNotification(FText::Format(LOCTEXT("RenameSessionFmt.InProgress", "Renaming Session {0} "), FText::FromString(SessionInfo.SessionName)))
+		);
+	
 	FText FailureReason = FText::GetEmpty();
 	const bool bSuccess = ServerInstance->GetConcertServer()->RenameSession(SessionId, NewName, FailureReason);
-	NotifyUserOfFinishedSessionAction(bSuccess,
-			bSuccess ? FText::Format(LOCTEXT("RenameSessionFmt", "Rename Session '{0}' as '{1}'"), FText::FromString(SessionInfo.SessionName), FText::FromString(NewName))
-				: FText::Format(LOCTEXT("FailedToArchivedSessionFmt", "Failed to rename Session '{0}' as '{1}'"), FText::FromString(SessionInfo.SessionName), FText::FromString(NewName)),
+
+	CloseAsyncNotification(
+			Notification,
+			bSuccess,
+			bSuccess
+				? FText::Format(LOCTEXT("RenameSessionFmt.Success", "Rename Session '{0}' as '{1}'"), FText::FromString(SessionInfo.SessionName), FText::FromString(NewName))
+				: FText::Format(LOCTEXT("RenameSessionFmt.Failure", "Failed to rename Session '{0}' as '{1}'"), FText::FromString(SessionInfo.SessionName), FText::FromString(NewName)),
 			FailureReason
-			);
+		);
 }
 
-void FConcertServerSessionBrowserController::DeleteSession(const FGuid& ServerAdminEndpointId, const FGuid& SessionId)
+bool FConcertServerSessionBrowserController::DeleteSession(const FGuid& ServerAdminEndpointId, const FGuid& SessionId, FText& FailureReason)
 {
 	const TOptional<FConcertSessionInfo> ArchivedInfo = GetArchivedSessionInfo(ServerAdminEndpointId, SessionId);
 	const TOptional<FConcertSessionInfo> LiveInfo = GetActiveSessionInfo(ServerAdminEndpointId, SessionId);
-	const TOptional<FString> SessionName = ArchivedInfo ? ArchivedInfo->SessionName : LiveInfo ? LiveInfo->SessionName : TOptional<FString>{};
-	if (SessionName)
+	if (const TOptional<FString> SessionName = GetSessionName(ServerAdminEndpointId, SessionId))
 	{
-		FText FailureReason;
-		const bool bSuccess = ServerInstance->GetConcertServer()->DestroySession(SessionId, FailureReason);
-		NotifyUserOfFinishedSessionAction(bSuccess,
-			bSuccess ? FText::Format(LOCTEXT("ArchivedSessionFmt", "Deleted Session '{0}'"), FText::FromString(*SessionName)) : FText::Format(LOCTEXT("FailedToArchivedSessionFmt", "Failed to archive Session '{0}'"), FText::FromString(LiveInfo->SessionName)),
-			FailureReason
-			);
+		return ServerInstance->GetConcertServer()->DestroySession(SessionId, FailureReason);
 	}
+	return false;
+}
+
+TOptional<FString> FConcertServerSessionBrowserController::GetSessionName(const FGuid& ServerAdminEndpointId, const FGuid& SessionId) const
+{
+	const TOptional<FConcertSessionInfo> ArchivedInfo = GetArchivedSessionInfo(ServerAdminEndpointId, SessionId);
+	const TOptional<FConcertSessionInfo> LiveInfo = GetActiveSessionInfo(ServerAdminEndpointId, SessionId);
+	const TOptional<FString> SessionName = ArchivedInfo ? ArchivedInfo->SessionName : LiveInfo ? LiveInfo->SessionName : TOptional<FString>{};
+	return SessionName;
+}
+
+FAsyncTaskNotificationConfig FConcertServerSessionBrowserController::MakeAsyncNotification(const FText& Title, const FText& Details)
+{
+	FAsyncTaskNotificationConfig NotificationConfig;
+	NotificationConfig.bIsHeadless = false;
+	NotificationConfig.bKeepOpenOnFailure = true;
+	NotificationConfig.TitleText = Title;
+	return NotificationConfig;
 }
 
-void FConcertServerSessionBrowserController::NotifyUserOfFinishedSessionAction(const bool bSuccess, const FText& Title, const FText& Details)
+void FConcertServerSessionBrowserController::CloseAsyncNotification(FAsyncTaskNotification& Notification, bool bSuccess, const FText& Title, const FText& Details)
 {
-	FNotificationInfo NotificationInfo(Title);
-	NotificationInfo.SubText = Details;
-	FSlateNotificationManager::Get().AddNotification(NotificationInfo);
+	Notification.SetComplete(Title, Details, bSuccess);
 
 	if (bSuccess)
 	{
diff --git a/Engine/Plugins/Developer/Concert/ConcertApp/MultiUserServer/Source/MultiUserServer/Private/Widgets/Browser/ConcertServerSessionBrowserController.h b/Engine/Plugins/Developer/Concert/ConcertApp/MultiUserServer/Source/MultiUserServer/Private/Widgets/Browser/ConcertServerSessionBrowserController.h
--- a/Engine/Plugins/Developer/Concert/ConcertApp/MultiUserServer/Source/MultiUserServer/Private/Widgets/Browser/ConcertServerSessionBrowserController.h
+++ b/Engine/Plugins/Developer/Concert/ConcertApp/MultiUserServer/Source/MultiUserServer/Private/Widgets/Browser/ConcertServerSessionBrowserController.h
@@ -3,6 +3,7 @@
 #pragma once
 
 #include "CoreMinimal.h"
+#include "Misc/AsyncTaskNotification.h"
 #include "Session/Browser/IConcertSessionBrowserController.h"
 #include "Widgets/IConcertComponent.h"
 
@@ -42,8 +43,7 @@ public:
 	virtual void RestoreSession(const FGuid& ServerAdminEndpointId, const FGuid& SessionId, const FString& RestoredName, const FConcertSessionFilter& SessionFilter) override;
 	virtual void RenameActiveSession(const FGuid& ServerAdminEndpointId, const FGuid& SessionId, const FString& NewName) override { RenameActiveSessionInternal(ServerAdminEndpointId, SessionId, NewName); }
 	virtual void RenameArchivedSession(const FGuid& ServerAdminEndpointId, const FGuid& SessionId, const FString& NewName) override { RenameArchivedSessionInternal(ServerAdminEndpointId, SessionId, NewName); }
-	virtual void DeleteActiveSession(const FGuid& ServerAdminEndpointId, const FGuid& SessionId) override { DeleteSession(ServerAdminEndpointId, SessionId); }
-	virtual void DeleteArchivedSession(const FGuid& ServerAdminEndpointId, const FGuid& SessionId) override { DeleteSession(ServerAdminEndpointId, SessionId); }
+	virtual void DeleteSessions(const FGuid& ServerAdminEndpointId, const TArray<FGuid>& SessionIds) override;
 
 	// The server operator always has permission for these actions:
 	virtual bool CanRenameActiveSession(const FGuid& ServerAdminEndpointId, const FGuid& SessionId) const override { return true; }
@@ -76,9 +76,11 @@ private:
 	void RenameActiveSessionInternal(const FGuid& ServerAdminEndpointId, const FGuid& SessionId, const FString& NewName);
 	void RenameArchivedSessionInternal(const FGuid& ServerAdminEndpointId, const FGuid& SessionId, const FString& NewName);
 	void RenameSessionInternal(const FGuid& SessionId, const FString& NewName, const FConcertSessionInfo& SessionInfo);
-	void DeleteSession(const FGuid& ServerAdminEndpointId, const FGuid& SessionId);
+	bool DeleteSession(const FGuid& ServerAdminEndpointId, const FGuid& SessionId, FText& FailureReason);
 
-	// Notifications
-	void NotifyUserOfFinishedSessionAction(const bool bSuccess, const FText& Title) { NotifyUserOfFinishedSessionAction(bSuccess, Title, FText::GetEmpty()); }
-	void NotifyUserOfFinishedSessionAction(const bool bSuccess, const FText& Title, const FText& Details);
+	TOptional<FString> GetSessionName(const FGuid& ServerAdminEndpointId, const FGuid& SessionId) const;
+
+	FAsyncTaskNotificationConfig MakeAsyncNotification(const FText& Title, const FText& Details = FText::GetEmpty());
+	void CloseAsyncNotification(FAsyncTaskNotification& Notification, bool bSuccess, const FText& Title, const FText& Details = FText::GetEmpty());
 };
+
diff --git a/Engine/Plugins/Developer/Concert/ConcertApp/MultiUserServer/Source/MultiUserServer/Private/Widgets/Browser/SConcertServerSessionBrowser.cpp b/Engine/Plugins/Developer/Concert/ConcertApp/MultiUserServer/Source/MultiUserServer/Private/Widgets/Browser/SConcertServerSessionBrowser.cpp
--- a/Engine/Plugins/Developer/Concert/ConcertApp/MultiUserServer/Source/MultiUserServer/Private/Widgets/Browser/SConcertServerSessionBrowser.cpp
+++ b/Engine/Plugins/Developer/Concert/ConcertApp/MultiUserServer/Source/MultiUserServer/Private/Widgets/Browser/SConcertServerSessionBrowser.cpp
@@ -5,17 +5,16 @@
 #include "ModalWindowManager.h"
 #include "MultiUserServerModule.h"
 #include "MultiUserServerUserSettings.h"
-#include "Session/Browser/SConcertSessionBrowser.h"
-#include "Widgets/Browser/ConcertServerSessionBrowserController.h"
-#include "Widgets/StatusBar/SConcertStatusBar.h"
 
 #include "Dialog/SMessageDialog.h"
-#include "Framework/Docking/TabManager.h"
+#include "Session/Browser/ConcertBrowserUtils.h"
 #include "Session/Browser/ConcertSessionItem.h"
+#include "Session/Browser/SConcertSessionBrowser.h"
+#include "Widgets/Browser/ConcertServerSessionBrowserController.h"
 #include "Widgets/ConcertServerTabs.h"
-#include "Widgets/Input/SSearchBox.h"
 #include "Widgets/Layout/SBorder.h"
 #include "Widgets/SBoxPanel.h"
+#include "Widgets/StatusBar/SConcertStatusBar.h"
 
 #define LOCTEXT_NAMESPACE "UnrealMultiUserUI"
 
@@ -55,10 +54,9 @@ TSharedRef<SWidget> SConcertServerSessionBrowser::MakeSessionTableView(const FAr
 	return SAssignNew(SessionBrowser, SConcertSessionBrowser, Controller.Pin().ToSharedRef(), SearchText)
 		.OnLiveSessionDoubleClicked(InArgs._DoubleClickLiveSession)
 		.OnArchivedSessionDoubleClicked(InArgs._DoubleClickArchivedSession)
-		.OnRequestedDeleteSession(this, &SConcertServerSessionBrowser::RequestDeleteSession)
+		.PostRequestedDeleteSession(this, &SConcertServerSessionBrowser::RequestDeleteSession)
 		// Pretend a modal dialog said no - RequestDeleteSession will show non-modal dialog
-		.CanDeleteArchivedSession_Lambda([](TSharedPtr<FConcertSessionItem>) { return false; })
-		.CanDeleteActiveSession_Lambda([](TSharedPtr<FConcertSessionItem>) { return false; })
+		.AskUserToDeleteSessions_Lambda([](auto) { return false; })
 		.ColumnVisibilitySnapshot(UMultiUserServerUserSettings::GetUserSettings()->GetSessionBrowserColumnVisibility())
 		.SaveColumnVisibilitySnapshot_Lambda([](const FColumnVisibilitySnapshot& Snapshot)
 		{
@@ -66,33 +64,46 @@ TSharedRef<SWidget> SConcertServerSessionBrowser::MakeSessionTableView(const FAr
 		});
 }
 
-void SConcertServerSessionBrowser::RequestDeleteSession(const TSharedPtr<FConcertSessionItem>& SessionItem) 
+void SConcertServerSessionBrowser::RequestDeleteSession(const TArray<TSharedPtr<FConcertSessionItem>>& SessionItems) 
 {
-	switch (SessionItem->Type)
+	const FText Message = [this, &SessionItems]()
 	{
-	case FConcertSessionItem::EType::ActiveSession:
-		DeleteActiveSessionWithFakeModalQuestion(SessionItem); break;
-	case FConcertSessionItem::EType::ArchivedSession:
-		DeleteArchivedSessionWithFakeModalQuestion(SessionItem); break;
-	default:
-		break;
-	}
+		if (SessionItems.Num() > 1)
+		{
+			return FText::Format(
+			LOCTEXT("DeletedMultipleDescription", "Deleting a session will force all connected clients to disconnect and all associated data to be removed.\n\nDelete {0} sessions?"),
+				SessionItems.Num()
+			);
+		}
+			
+		switch (SessionItems[0]->Type)
+		{
+		case FConcertSessionItem::EType::ActiveSession:
+			return FText::Format(
+				LOCTEXT("DeletedActiveDescription", "There {0}|plural(one=is,other=are) {0} connected {0}|plural(one=client,other=clients) in the current session.\nDeleting a session will force all connected clients to disconnect.\n\nDelete {1}?"),
+				Controller.Pin()->GetNumConnectedClients(SessionItems[0]->SessionId),
+				FText::FromString(SessionItems[0]->SessionName)
+			);
+		case FConcertSessionItem::EType::ArchivedSession:
+			return FText::Format(
+				LOCTEXT("DeleteArchivedDescription", "Deleting a session will cause all associated data to be removed.\n\nDelete {0}?"),
+				FText::FromString(SessionItems[0]->SessionName)
+				);
+		default:
+			checkNoEntry();
+			return FText::GetEmpty();
+		}
+	}();
+	DeleteSessionsWithFakeModalQuestion(Message, SessionItems);
 }
 
-void SConcertServerSessionBrowser::DeleteArchivedSessionWithFakeModalQuestion(const TSharedPtr<FConcertSessionItem>& SessionItem)
+void SConcertServerSessionBrowser::DeleteSessionsWithFakeModalQuestion(const FText& Message, const TArray<TSharedPtr<FConcertSessionItem>>& SessionItems)
 {
-	const FText Message = FText::Format(
-	LOCTEXT("DeleteArchivedDescription", "Deleting a session will cause all associated data to be removed.\n\nDelete {0}?"),
-		FText::FromString(SessionItem->SessionName)
-		);
-
-	auto DeleteArchived = [WeakController = TWeakPtr<FConcertServerSessionBrowserController>(Controller), WeakSessionItem = TWeakPtr<FConcertSessionItem>(SessionItem)]()
+	auto DeleteArchived = [WeakController = TWeakPtr<FConcertServerSessionBrowserController>(Controller), SessionItems]()
 	{
-		const TSharedPtr<FConcertServerSessionBrowserController> PinnedController = WeakController.Pin();
-		const TSharedPtr<FConcertSessionItem> PinnedSessionItem = WeakSessionItem.Pin();
-		if (PinnedController && PinnedSessionItem)
+		if (const TSharedPtr<FConcertServerSessionBrowserController> PinnedController = WeakController.Pin())
 		{
-			PinnedController->DeleteArchivedSession(PinnedSessionItem->ServerAdminEndpointId, PinnedSessionItem->SessionId);
+			ConcertBrowserUtils::RequestItemDeletion(*PinnedController.Get(), SessionItems);
 		}
 	};
 	const TSharedRef<SMessageDialog> Dialog = SNew(SMessageDialog)
@@ -113,40 +124,4 @@ void SConcertServerSessionBrowser::DeleteArchivedSessionWithFakeModalQuestion(co
 		->ShowFakeModalWindow(Dialog);
 }
 
-void SConcertServerSessionBrowser::DeleteActiveSessionWithFakeModalQuestion(const TSharedPtr<FConcertSessionItem>& SessionItem)
-{
-	const int32 NumUsers = Controller.Pin()->GetNumConnectedClients(SessionItem->SessionId);
-	const FText Message = FText::Format(
-		LOCTEXT("DeletedActiveDescription", "There {0}|plural(one=is,other=are) {0} connected {0}|plural(one=client,other=clients) in the current session.\nDeleting a session will force all connected clients to disconnect.\n\nDelete {1}?"),
-		NumUsers,
-		FText::FromString(SessionItem->SessionName)
-		);
-
-	auto DeleteActive = [WeakController = TWeakPtr<FConcertServerSessionBrowserController>(Controller), WeakSessionItem = TWeakPtr<FConcertSessionItem>(SessionItem)]()
-	{
-		const TSharedPtr<FConcertServerSessionBrowserController> PinnedController = WeakController.Pin();
-		const TSharedPtr<FConcertSessionItem> PinnedSessionItem = WeakSessionItem.Pin();
-		if (PinnedController && PinnedSessionItem)
-		{
-			PinnedController->DeleteActiveSession(PinnedSessionItem->ServerAdminEndpointId, PinnedSessionItem->SessionId);
-		}
-	};
-	const TSharedRef<SMessageDialog> Dialog = SNew(SMessageDialog)
-		.Title(LOCTEXT("DisconnectUsersTitle", "Force Users to Disconnect?"))
-		.IconBrush("Icons.WarningWithColor.Large")
-		.Message(Message)
-		.UseScrollBox(false)
-		.Buttons({
-			SMessageDialog::FButton(LOCTEXT("DeleteActiveButton", "Delete"))
-				.SetPrimary(true)
-				.SetOnClicked(FSimpleDelegate::CreateLambda(DeleteActive))
-				.SetFocus(),
-			SMessageDialog::FButton(LOCTEXT("CancelButton", "Cancel"))
-		});
-
-	UE::MultiUserServer::FConcertServerUIModule::Get()
-		.GetModalWindowManager()
-		->ShowFakeModalWindow(Dialog);
-}
-
 #undef LOCTEXT_NAMESPACE
diff --git a/Engine/Plugins/Developer/Concert/ConcertApp/MultiUserServer/Source/MultiUserServer/Private/Widgets/Browser/SConcertServerSessionBrowser.h b/Engine/Plugins/Developer/Concert/ConcertApp/MultiUserServer/Source/MultiUserServer/Private/Widgets/Browser/SConcertServerSessionBrowser.h
--- a/Engine/Plugins/Developer/Concert/ConcertApp/MultiUserServer/Source/MultiUserServer/Private/Widgets/Browser/SConcertServerSessionBrowser.h
+++ b/Engine/Plugins/Developer/Concert/ConcertApp/MultiUserServer/Source/MultiUserServer/Private/Widgets/Browser/SConcertServerSessionBrowser.h
@@ -47,7 +47,6 @@ private:
 	
 	TSharedRef<SWidget> MakeSessionTableView(const FArguments& InArgs);
 
-	void RequestDeleteSession(const TSharedPtr<FConcertSessionItem>& SessionItem);
-	void DeleteArchivedSessionWithFakeModalQuestion(const TSharedPtr<FConcertSessionItem>& SessionItem);
-	void DeleteActiveSessionWithFakeModalQuestion(const TSharedPtr<FConcertSessionItem>& SessionItem);
+	void RequestDeleteSession(const TArray<TSharedPtr<FConcertSessionItem>>& SessionItems);
+	void DeleteSessionsWithFakeModalQuestion(const FText& Message, const TArray<TSharedPtr<FConcertSessionItem>>& SessionItems);
 };
diff --git a/Engine/Plugins/Developer/Concert/ConcertMain/Source/Concert/Private/ConcertClient.cpp b/Engine/Plugins/Developer/Concert/ConcertMain/Source/Concert/Private/ConcertClient.cpp
--- a/Engine/Plugins/Developer/Concert/ConcertMain/Source/Concert/Private/ConcertClient.cpp
+++ b/Engine/Plugins/Developer/Concert/ConcertMain/Source/Concert/Private/ConcertClient.cpp
@@ -7,6 +7,8 @@
 #include "ConcertLogger.h"
 #include "ConcertLogGlobal.h"
 
+#include "Algo/Transform.h"
+
 #include "Containers/Ticker.h"
 #include "Misc/App.h"
 #include "Misc/Paths.h"
@@ -1000,6 +1002,52 @@ TFuture<EConcertResponseCode> FConcertClient::DeleteSession(const FGuid& ServerA
 		});
 }
 
+TFuture<FConcertAdmin_BatchDeleteSessionResponse> FConcertClient::BatchDeleteSessions(const FGuid& ServerAdminEndpointId, const FConcertBatchDeleteSessionsArgs& BatchDeletionArgs)
+{
+	FConcertAdmin_BatchDeleteSessionRequest DeleteSessionRequest;
+	DeleteSessionRequest.SessionIds = BatchDeletionArgs.SessionIds;
+	DeleteSessionRequest.Flags = BatchDeletionArgs.Flags;
+
+	// Fill the information for the client identification
+	DeleteSessionRequest.UserName = ClientInfo.UserName;
+	DeleteSessionRequest.DeviceName = ClientInfo.DeviceName;
+
+	FAsyncTaskNotificationConfig NotificationConfig;
+	NotificationConfig.bIsHeadless = Settings->bIsHeadless;
+	NotificationConfig.bKeepOpenOnFailure = true;
+	NotificationConfig.TitleText = LOCTEXT("DeletingSessions", "Deleting Sessions...");
+	NotificationConfig.LogCategory = ConcertUtil::GetLogConcertPtr();
+
+	FAsyncTaskNotification Notification(NotificationConfig);
+
+	return ClientAdminEndpoint->SendRequest<FConcertAdmin_BatchDeleteSessionRequest, FConcertAdmin_BatchDeleteSessionResponse>(DeleteSessionRequest, ServerAdminEndpointId)
+		.Next([this, NumRequested = DeleteSessionRequest.SessionIds.Num(), Notification = MoveTemp(Notification)](const FConcertAdmin_BatchDeleteSessionResponse& RequestResponse) mutable
+		{
+			if (RequestResponse.ResponseCode == EConcertResponseCode::Success)
+			{
+				const bool bDeletedAll = RequestResponse.DeletedItems.Num() == NumRequested;
+				const FText Message = bDeletedAll
+					? FText::Format(LOCTEXT("DeletedSessionsFmt.All", "Deleted {0} Sessions"), RequestResponse.DeletedItems.Num())
+					: FText::Format(LOCTEXT("DeletedSessionsFmt.Some", "Deleted {0} of {1} Sessions"), RequestResponse.DeletedItems.Num(), NumRequested);
+				const FText ProgressText = bDeletedAll
+					? FText::GetEmpty()
+					: [&RequestResponse]()
+					{
+						TArray<FString> SessionNames;
+						Algo::Transform(RequestResponse.NotOwnedByClient, SessionNames, [](const FDeletedSessionInfo& Skipped){ return Skipped.SessionName; });
+						return FText::FromString(FString::Join(SessionNames, TEXT(", ")));
+					}();
+				Notification.SetComplete(Message, ProgressText, true);
+			}
+			else
+			{
+				Notification.SetComplete(LOCTEXT("FailedToDeleteSessionsFmt", "Failed to Delete Sessions"), RequestResponse.Reason, false);
+			}
+
+			return RequestResponse;
+		});
+}
+
 void FConcertClient::DisconnectSession()
 {
 	// We don't want the client to get automatically reconnected to it's default session
diff --git a/Engine/Plugins/Developer/Concert/ConcertMain/Source/Concert/Private/ConcertClient.h b/Engine/Plugins/Developer/Concert/ConcertMain/Source/Concert/Private/ConcertClient.h
--- a/Engine/Plugins/Developer/Concert/ConcertMain/Source/Concert/Private/ConcertClient.h
+++ b/Engine/Plugins/Developer/Concert/ConcertMain/Source/Concert/Private/ConcertClient.h
@@ -84,6 +84,7 @@ public:
 	virtual TFuture<EConcertResponseCode> ArchiveSession(const FGuid& ServerAdminEndpointId, const FConcertArchiveSessionArgs& ArchiveSessionArgs) override;
 	virtual TFuture<EConcertResponseCode> RenameSession(const FGuid& ServerAdminEndpointId, const FGuid& SessionId, const FString& NewName) override;
 	virtual TFuture<EConcertResponseCode> DeleteSession(const FGuid& ServerAdminEndpointId, const FGuid& SessionId) override;
+	virtual TFuture<FConcertAdmin_BatchDeleteSessionResponse> BatchDeleteSessions(const FGuid& ServerAdminEndpointId, const FConcertBatchDeleteSessionsArgs& BatchDeletionArgs) override;
 	virtual void DisconnectSession() override;
 	virtual void ResumeSession() override;
 	virtual void SuspendSession() override;
diff --git a/Engine/Plugins/Developer/Concert/ConcertMain/Source/Concert/Private/ConcertServer.cpp b/Engine/Plugins/Developer/Concert/ConcertMain/Source/Concert/Private/ConcertServer.cpp
--- a/Engine/Plugins/Developer/Concert/ConcertMain/Source/Concert/Private/ConcertServer.cpp
+++ b/Engine/Plugins/Developer/Concert/ConcertMain/Source/Concert/Private/ConcertServer.cpp
@@ -21,6 +21,8 @@
 #include "Backends/JsonStructDeserializerBackend.h"
 #include "HAL/FileManager.h"
 
+#include "Templates/NonNullPointer.h"
+
 #define LOCTEXT_NAMESPACE "ConcertServer"
 
 namespace ConcertServerUtil
@@ -171,6 +173,7 @@ void FConcertServer::Startup()
 		ServerAdminEndpoint->RegisterRequestHandler<FConcertAdmin_ArchiveSessionRequest, FConcertAdmin_ArchiveSessionResponse>(this, &FConcertServer::HandleArchiveSessionRequest);
 		ServerAdminEndpoint->RegisterRequestHandler<FConcertAdmin_RenameSessionRequest, FConcertAdmin_RenameSessionResponse>(this, &FConcertServer::HandleRenameSessionRequest);
 		ServerAdminEndpoint->RegisterRequestHandler<FConcertAdmin_DeleteSessionRequest, FConcertAdmin_DeleteSessionResponse>(this, &FConcertServer::HandleDeleteSessionRequest);
+		ServerAdminEndpoint->RegisterRequestHandler<FConcertAdmin_BatchDeleteSessionRequest, FConcertAdmin_BatchDeleteSessionResponse>(this, &FConcertServer::HandleBatchDeleteSessionRequest);
 
 		ServerAdminEndpoint->RegisterRequestHandler<FConcertAdmin_GetAllSessionsRequest, FConcertAdmin_GetAllSessionsResponse>(this, &FConcertServer::HandleGetAllSessionsRequest);
 		ServerAdminEndpoint->RegisterRequestHandler<FConcertAdmin_GetLiveSessionsRequest, FConcertAdmin_GetSessionsResponse>(this, &FConcertServer::HandleGetLiveSessionsRequest);
@@ -1236,6 +1239,92 @@ FConcertAdmin_RenameSessionResponse FConcertServer::RenameSessionInternal(const
 	return ResponseData;
 }
 
+TFuture<FConcertAdmin_BatchDeleteSessionResponse> FConcertServer::HandleBatchDeleteSessionRequest(const FConcertMessageContext& Context)
+{
+	FConcertAdmin_BatchDeleteSessionResponse ResponseData;
+	ResponseData.ResponseCode = EConcertResponseCode::Failed;
+	
+	const FConcertAdmin_BatchDeleteSessionRequest& Request = *Context.GetMessage<FConcertAdmin_BatchDeleteSessionRequest>();
+	TMap<FGuid, FString> SessionNames;
+	if (ValidateBatchDeletionRequest(Request, ResponseData, SessionNames))
+	{
+		ResponseData.ResponseCode = EConcertResponseCode::Success;
+		FConcertAdmin_DeleteSessionRequest DeleteSingleSessionRequest;
+		for (const FGuid& SessionToDelete : Request.SessionIds)
+		{
+			const bool bSkip = ResponseData.NotOwnedByClient.ContainsByPredicate([&SessionToDelete](const FDeletedSessionInfo& Info ){ return Info.SessionId == SessionToDelete; });
+			if (bSkip)
+			{
+				continue;
+			}
+			
+			DeleteSingleSessionRequest.SessionId = SessionToDelete;
+			const FConcertAdmin_DeleteSessionResponse DeleteResponse = DeleteSessionInternal(DeleteSingleSessionRequest, false);
+			if (DeleteResponse.ResponseCode == EConcertResponseCode::Success)
+			{
+				ResponseData.DeletedItems.Add({ SessionToDelete, SessionNames[SessionToDelete] });
+			}
+			else
+			{
+				// We may already have deleted some files ... maybe we should restore them in the future...
+				ResponseData.ResponseCode = EConcertResponseCode::Failed;
+				break;
+			}
+		}
+	}
+
+	return FConcertAdmin_BatchDeleteSessionResponse::AsFuture(MoveTemp(ResponseData));
+}
+
+bool FConcertServer::ValidateBatchDeletionRequest(const FConcertAdmin_BatchDeleteSessionRequest& Request, FConcertAdmin_BatchDeleteSessionResponse& OutResponse, TMap<FGuid, FString>& PreparedSessionInfo) const
+{
+	TArray<FDeletedSessionInfo> NotOwnedByClient;
+	for (const FGuid& SessionToDelete : Request.SessionIds)
+	{
+		if (TSharedPtr<IConcertServerSession> ServerSession = GetLiveSession(SessionToDelete))
+		{
+			const bool bHasPermission = IsRequestFromSessionOwner(ServerSession, Request.UserName, Request.DeviceName);
+			if (!bHasPermission && (Request.Flags & EBatchSessionDeletionFlags::SkipForbiddenSessions) != EBatchSessionDeletionFlags::Strict)
+			{
+				NotOwnedByClient.Add({ SessionToDelete, ServerSession->GetName() });
+			}
+			else if (!bHasPermission)
+			{
+				OutResponse.Reason = LOCTEXT("Error_BatchDelete_InvalidPerms_NotOwner", "Not the session owner.");
+				UE_LOG(LogConcert, Display, TEXT("User %s failed to delete live session '%s' (Id: %s, Owner: %s, Reason: %s)"), *Request.UserName, *ServerSession->GetName(), *SessionToDelete.ToString(), *ServerSession->GetSessionInfo().OwnerUserName, *OutResponse.Reason.ToString());
+				return false;
+			}
+
+			PreparedSessionInfo.Add(SessionToDelete, ServerSession->GetName());
+		}
+		else if (const FConcertSessionInfo* ArchivedSessionInfo = ArchivedSessions.Find(SessionToDelete))
+		{
+			const bool bHasPermission = IsRequestFromSessionOwner(*ArchivedSessionInfo, Request.UserName, Request.DeviceName);
+			if (!bHasPermission && (Request.Flags & EBatchSessionDeletionFlags::SkipForbiddenSessions) != EBatchSessionDeletionFlags::Strict)
+			{
+				NotOwnedByClient.Add({ SessionToDelete, ServerSession->GetName() });
+			}
+			else if (!bHasPermission)
+			{
+				OutResponse.Reason = LOCTEXT("Error_BatchDelete_InvalidPerms_NotOwner", "Not the session owner.");
+				UE_LOG(LogConcert, Display, TEXT("User %s failed to delete live session '%s' (Id: %s, Owner: %s, Reason: %s)"), *Request.UserName, *ArchivedSessionInfo->SessionName, *SessionToDelete.ToString(), *ServerSession->GetSessionInfo().OwnerUserName, *OutResponse.Reason.ToString());
+				return false;
+			}
+			
+			PreparedSessionInfo.Add(SessionToDelete, ArchivedSessionInfo->SessionName);
+		}
+		else
+		{
+			OutResponse.Reason = FText::Format(LOCTEXT("Error_BatchDelete_SessionDoesNotExist", "Session ID {0} does not exist."), FText::FromString(SessionToDelete.ToString()));
+			UE_LOG(LogConcert, Display, TEXT("User %s failed to delete session (Id: %s, Reason: %s)"), *Request.UserName, *SessionToDelete.ToString(), *OutResponse.Reason.ToString());
+			return false;
+		}
+	}
+
+	OutResponse.NotOwnedByClient = MoveTemp(NotOwnedByClient);
+	return true;
+}
+
 TFuture<FConcertAdmin_DeleteSessionResponse> FConcertServer::HandleDeleteSessionRequest(const FConcertMessageContext & Context)
 {
 	return FConcertAdmin_DeleteSessionResponse::AsFuture(DeleteSessionInternal(*Context.GetMessage<FConcertAdmin_DeleteSessionRequest>(), /*bCheckPermission*/true));
@@ -1391,16 +1480,21 @@ bool FConcertServer::CanJoinSession(const TSharedPtr<IConcertServerSession>& Ser
 	return true;
 }
 
-bool FConcertServer::IsRequestFromSessionOwner(const TSharedPtr<IConcertServerSession>& SessionToDelete, const FString& FromUserName, const FString& FromDeviceName)
+bool FConcertServer::IsRequestFromSessionOwner(const TSharedPtr<IConcertServerSession>& SessionToDelete, const FString& FromUserName, const FString& FromDeviceName) const
 {
 	if (SessionToDelete)
 	{
 		const FConcertSessionInfo& SessionInfo = SessionToDelete->GetSessionInfo();
-		return SessionInfo.OwnerUserName == FromUserName && SessionInfo.OwnerDeviceName == FromDeviceName;
+		return IsRequestFromSessionOwner(SessionInfo, FromUserName, FromDeviceName);
 	}
 	return false;
 }
 
+bool FConcertServer::IsRequestFromSessionOwner(const FConcertSessionInfo& SessionInfo, const FString& FromUserName, const FString& FromDeviceName) const
+{
+	return SessionInfo.OwnerUserName == FromUserName && SessionInfo.OwnerDeviceName == FromDeviceName;
+}
+
 TSharedPtr<IConcertServerSession> FConcertServer::CreateLiveSession(const FConcertSessionInfo& SessionInfo, const FConcertServerSessionRepository& InRepository)
 {
 	check(SessionInfo.SessionId.IsValid() && !SessionInfo.SessionName.IsEmpty());
diff --git a/Engine/Plugins/Developer/Concert/ConcertMain/Source/Concert/Private/ConcertServer.h b/Engine/Plugins/Developer/Concert/ConcertMain/Source/Concert/Private/ConcertServer.h
--- a/Engine/Plugins/Developer/Concert/ConcertMain/Source/Concert/Private/ConcertServer.h
+++ b/Engine/Plugins/Developer/Concert/ConcertMain/Source/Concert/Private/ConcertServer.h
@@ -90,6 +90,12 @@ private:
 	/**  */
 	FConcertAdmin_RenameSessionResponse RenameSessionInternal(const FConcertAdmin_RenameSessionRequest& Request, bool bCheckPermission);
 
+	/**  */
+	TFuture<FConcertAdmin_BatchDeleteSessionResponse> HandleBatchDeleteSessionRequest(const FConcertMessageContext& Context);
+
+	/** Checks whether the request is valid to serve. Writes any skipped sessions (if any) and error reasons into OutResponse. */
+	bool ValidateBatchDeletionRequest(const FConcertAdmin_BatchDeleteSessionRequest& Request, FConcertAdmin_BatchDeleteSessionResponse& OutResponse, TMap<FGuid, FString>& OutPreparedSessionNames) const;
+	
 	/**  */
 	TFuture<FConcertAdmin_DeleteSessionResponse> HandleDeleteSessionRequest(const FConcertMessageContext& Context);
 
@@ -126,11 +132,11 @@ private:
 	/** */
 	bool CanJoinSession(const TSharedPtr<IConcertServerSession>& ServerSession, const FConcertSessionSettings& SessionSettings, const FConcertSessionVersionInfo& SessionVersionInfo, FText* OutFailureReason = nullptr);
 
-	/** 
-	 * Validate that the request come form the owner of the session that he want to delete/rename, etc.
-	 */
-	bool IsRequestFromSessionOwner(const TSharedPtr<IConcertServerSession>& Session, const FString& FromUserName, const FString& FromDeviceName);
-
+	/** Validate that the request come form the owner of the session that he want to delete/rename, etc. */
+	bool IsRequestFromSessionOwner(const TSharedPtr<IConcertServerSession>& Session, const FString& FromUserName, const FString& FromDeviceName) const;
+	/** Validate that the request come form the owner of the session that he want to delete/rename, etc. */
+	bool IsRequestFromSessionOwner(const FConcertSessionInfo& SessionInfo, const FString& FromUserName, const FString& FromDeviceName) const;
+	
 	/**  */
 	TSharedPtr<IConcertServerSession> CreateLiveSession(const FConcertSessionInfo& SessionInfo, const FConcertServerSessionRepository& InRepository);
 
diff --git a/Engine/Plugins/Developer/Concert/ConcertMain/Source/Concert/Public/ConcertMessages.h b/Engine/Plugins/Developer/Concert/ConcertMain/Source/Concert/Public/ConcertMessages.h
--- a/Engine/Plugins/Developer/Concert/ConcertMain/Source/Concert/Public/ConcertMessages.h
+++ b/Engine/Plugins/Developer/Concert/ConcertMain/Source/Concert/Public/ConcertMessages.h
@@ -405,7 +405,7 @@ struct FConcertAdmin_RenameSessionResponse : public FConcertResponseData
 };
 
 
-/** Delete a live session. */
+/** Delete an archived or live session. */
 USTRUCT()
 struct FConcertAdmin_DeleteSessionRequest : public FConcertRequestData
 {
@@ -437,6 +437,72 @@ struct FConcertAdmin_DeleteSessionResponse : public FConcertResponseData
 	FString SessionName;
 };
 
+UENUM()
+enum class EBatchSessionDeletionFlags
+{
+	/** Any error in the request will make the entire operation fail. */
+	Strict,
+
+	/**
+	 * If the request contains sessions that is not owned by the requesting client, skip them.
+	 * Skipped sessions are added to FConcertAdmin_BatchDeleteSessionResponse::NotOwnedByClient.
+	 * This avoids clients having to request ownership info before ensuring an atomic and consistent operation.
+	 */
+	SkipForbiddenSessions = 0x01
+};
+ENUM_CLASS_FLAGS(EBatchSessionDeletionFlags)
+
+/** Deletes several archived and/or live sessions. */
+USTRUCT()
+struct FConcertAdmin_BatchDeleteSessionRequest : public FConcertRequestData
+{
+	GENERATED_BODY()
+
+	/** The ID of the session to delete. */
+	UPROPERTY(VisibleAnywhere, Category = "Concert Message")
+	TSet<FGuid> SessionIds;
+
+	/** How the operation input is supposed to be interpreted */
+	UPROPERTY(VisibleAnywhere, Category = "Concert Message")
+	EBatchSessionDeletionFlags Flags;
+
+	//For now only the user name and device name of the client is used to id him as the owner of a session
+	UPROPERTY(VisibleAnywhere, Category = "Concert Message")
+	FString UserName;
+
+	UPROPERTY(VisibleAnywhere, Category = "Concert Message")
+	FString DeviceName;
+};
+
+USTRUCT()
+struct FDeletedSessionInfo
+{
+	GENERATED_BODY()
+	
+	/** The ID of the session that was requested to be deleted. */
+	UPROPERTY(VisibleAnywhere, Category = "Concert Message")
+	FGuid SessionId;
+
+	/** The name of the session that was was requested to be deleted. */
+	UPROPERTY(VisibleAnywhere, Category = "Concert Message")
+	FString SessionName;
+};
+
+/** Answer by server */
+USTRUCT()
+struct FConcertAdmin_BatchDeleteSessionResponse : public FConcertResponseData
+{
+	GENERATED_BODY()
+
+	/** The IDs of the session that were requested to be deleted. */
+	UPROPERTY(VisibleAnywhere, Category = "Concert Message")
+	TArray<FDeletedSessionInfo> DeletedItems;
+
+	/** Contains the sessions that were skipped if FConcertAdmin_BatchDeleteSessionRequest::Flags has the SkipForbiddenSessions flag set. */
+	UPROPERTY(VisibleAnywhere, Category = "Concert Message")
+	TArray<FDeletedSessionInfo> NotOwnedByClient;
+};
+
 USTRUCT()
 struct FConcertAdmin_GetSessionClientsRequest : public FConcertRequestData
 {
diff --git a/Engine/Plugins/Developer/Concert/ConcertMain/Source/Concert/Public/IConcertClient.h b/Engine/Plugins/Developer/Concert/ConcertMain/Source/Concert/Public/IConcertClient.h
--- a/Engine/Plugins/Developer/Concert/ConcertMain/Source/Concert/Public/IConcertClient.h
+++ b/Engine/Plugins/Developer/Concert/ConcertMain/Source/Concert/Public/IConcertClient.h
@@ -143,6 +143,12 @@ struct FConcertArchiveSessionArgs
 	FConcertSessionFilter SessionFilter;
 };
 
+struct FConcertBatchDeleteSessionsArgs
+{
+	TSet<FGuid> SessionIds;
+	EBatchSessionDeletionFlags Flags;
+};
+
 /** Interface for Concert client */
 class IConcertClient
 {
@@ -337,6 +343,15 @@ public:
 	 */
 	virtual TFuture<EConcertResponseCode> DeleteSession(const FGuid& ServerAdminEndpointId, const FGuid& SessionId) = 0;
 
+	/**
+	 * Deletes several live or archives sessions from the server. The client will locally filter SessionIds and request deletion of only those sessions which the client is an owner of.
+	 * The server will still internally validate the client has permission and possibly refuse if there's a mismatch.
+	 * @param ServerAdminEndpointId	The Id of the Concert Server query endpoint
+	 * @param BatchDeletionArgs		The arguments that will be used for batch deleting the session(s)
+	 * @return A future that will contain the server's response AND which sessions were filtered out by the client before sending.
+	 */
+	virtual TFuture<FConcertAdmin_BatchDeleteSessionResponse> BatchDeleteSessions(const FGuid& ServerAdminEndpointId, const FConcertBatchDeleteSessionsArgs& BatchDeletionArgs) = 0;
+	
 	/** 
 	 * Disconnect from the current session.
 	 */
diff --git a/Engine/Plugins/Developer/Concert/ConcertUI/ConcertSharedSlate/Source/ConcertSharedSlate/Private/Session/Browser/SConcertSessionBrowser.cpp b/Engine/Plugins/Developer/Concert/ConcertUI/ConcertSharedSlate/Source/ConcertSharedSlate/Private/Session/Browser/SConcertSessionBrowser.cpp
--- a/Engine/Plugins/Developer/Concert/ConcertUI/ConcertSharedSlate/Source/ConcertSharedSlate/Private/Session/Browser/SConcertSessionBrowser.cpp
+++ b/Engine/Plugins/Developer/Concert/ConcertUI/ConcertSharedSlate/Source/ConcertSharedSlate/Private/Session/Browser/SConcertSessionBrowser.cpp
@@ -7,6 +7,7 @@
 #include "ConcertHeaderRowUtils.h"
 
 #include "Algo/Find.h"
+#include "Algo/Transform.h"
 
 #include "Session/Browser/ConcertBrowserUtils.h"
 #include "Session/Browser/ConcertSessionBrowserSettings.h"
@@ -56,9 +57,8 @@ void SConcertSessionBrowser::Construct(const FArguments& InArgs, TSharedRef<ICon
 	{
 		OnArchivedSessionDoubleClicked = FSessionDelegate::CreateSP(this, &SConcertSessionBrowser::InsertRestoreSessionAsEditableRowInternal);
 	}
-	OnRequestedDeleteSession = InArgs._OnRequestedDeleteSession;
-	CanDeleteArchivedSession = InArgs._CanDeleteArchivedSession;
-	CanDeleteActiveSession = InArgs._CanDeleteActiveSession;
+	PostRequestedDeleteSession = InArgs._PostRequestedDeleteSession;
+	AskUserToDeleteSessions = InArgs._AskUserToDeleteSessions;
 
 	// Setup search filter.
 	SearchedText = InSearchText; // Reload a previous search text (in any). Useful to remember searched text between join/leave sessions, but not persistent if the tab is closed.
@@ -356,10 +356,9 @@ TSharedRef<SWidget> SConcertSessionBrowser::MakeSessionTableView(const FArgument
 	SecondarySortMode = EColumnSortMode::Ascending;
 
 	SessionsView = SNew(SListView<TSharedPtr<FConcertSessionItem>>)
-		.SelectionMode(ESelectionMode::Single)
 		.ListItemsSource(&Sessions)
 		.OnGenerateRow(this, &SConcertSessionBrowser::OnGenerateSessionRowWidget)
-		.SelectionMode(ESelectionMode::Single)
+		.SelectionMode(ESelectionMode::Multi)
 		.OnSelectionChanged(this, &SConcertSessionBrowser::OnSessionSelectionChanged)
 		.OnContextMenuOpening(this, &SConcertSessionBrowser::MakeContextualMenu)
 		.HeaderRow(
@@ -1054,12 +1053,7 @@ FReply SConcertSessionBrowser::OnArchiveButtonClicked()
 
 FReply SConcertSessionBrowser::OnDeleteButtonClicked()
 {
-	TArray<TSharedPtr<FConcertSessionItem>> SelectedItems = SessionsView->GetSelectedItems();
-	for (const TSharedPtr<FConcertSessionItem>& Item : SelectedItems)
-	{
-		RequestDeleteSession(Item);
-	}
-
+	RequestDeleteSessions(SessionsView->GetSelectedItems());
 	return FReply::Handled();
 }
 
@@ -1093,10 +1087,7 @@ FReply SConcertSessionBrowser::OnKeyDown(const FGeometry& MyGeometry, const FKey
 
 	if (InKeyEvent.GetKey() == EKeys::Delete && !EditableSessionRow.IsValid()) // Delete selected row(s) unless the selected row is an 'editable' one.
 	{
-		for (TSharedPtr<FConcertSessionItem>& Item : SessionsView->GetSelectedItems())
-		{
-			RequestDeleteSession(Item);
-		}
+		RequestDeleteSessions(SessionsView->GetSelectedItems());
 		return FReply::Handled();
 	}
 	else if (InKeyEvent.GetKey() == EKeys::Escape && EditableSessionRow.IsValid()) // Cancel 'new session', 'archive session' or 'restore session' action.
@@ -1151,24 +1142,13 @@ void SConcertSessionBrowser::RequestRenameSession(const TSharedPtr<FConcertSessi
 	RenamedItem->SessionName = NewName;
 }
 
-void SConcertSessionBrowser::RequestDeleteSession(const TSharedPtr<FConcertSessionItem>& DeletedItem)
+void SConcertSessionBrowser::RequestDeleteSessions(const TArray<TSharedPtr<FConcertSessionItem>>& ItemsToDelete)
 {
-	if (DeletedItem->Type == FConcertSessionItem::EType::ActiveSession)
+	if (!AskUserToDeleteSessions.IsBound() || AskUserToDeleteSessions.Execute(ItemsToDelete))
 	{
-		if (!CanDeleteActiveSession.IsBound() || CanDeleteActiveSession.Execute(DeletedItem))
-		{
-			GetController()->DeleteActiveSession(DeletedItem->ServerAdminEndpointId, DeletedItem->SessionId);
-		}
-	}
-	else if (DeletedItem->Type == FConcertSessionItem::EType::ArchivedSession)
-	{
-		if (!CanDeleteArchivedSession.IsBound() || CanDeleteArchivedSession.Execute(DeletedItem))
-		{
-			GetController()->DeleteArchivedSession(DeletedItem->ServerAdminEndpointId, DeletedItem->SessionId);
-		}
+		ConcertBrowserUtils::RequestItemDeletion(*GetController().Get(), ItemsToDelete);
 	}
-	
-	OnRequestedDeleteSession.ExecuteIfBound(DeletedItem);
+	PostRequestedDeleteSession.ExecuteIfBound(ItemsToDelete);
 }
 
 #undef LOCTEXT_NAMESPACE
diff --git a/Engine/Plugins/Developer/Concert/ConcertUI/ConcertSharedSlate/Source/ConcertSharedSlate/Public/Session/Browser/ConcertBrowserUtils.h b/Engine/Plugins/Developer/Concert/ConcertUI/ConcertSharedSlate/Source/ConcertSharedSlate/Public/Session/Browser/ConcertBrowserUtils.h
--- a/Engine/Plugins/Developer/Concert/ConcertUI/ConcertSharedSlate/Source/ConcertSharedSlate/Public/Session/Browser/ConcertBrowserUtils.h
+++ b/Engine/Plugins/Developer/Concert/ConcertUI/ConcertSharedSlate/Source/ConcertSharedSlate/Public/Session/Browser/ConcertBrowserUtils.h
@@ -5,15 +5,17 @@
 #include "CoreMinimal.h"
 
 #include "ConcertMessageData.h"
+#include "IConcertSessionBrowserController.h"
+#include "Session/Browser/ConcertSessionItem.h"
 
 #include "EditorFontGlyphs.h"
-#include "SNegativeActionButton.h"
-#include "SPositiveActionButton.h"
-#include "Styling/AppStyle.h"
 #include "Internationalization/Regex.h"
 #include "Layout/Visibility.h"
 #include "Misc/Attribute.h"
 #include "Misc/AsyncTaskNotification.h"
+#include "SNegativeActionButton.h"
+#include "SPositiveActionButton.h"
+#include "Styling/AppStyle.h"
 #include "Styling/SlateColor.h"
 #include "Widgets/Images/SImage.h"
 #include "Widgets/Input/SButton.h"
@@ -101,4 +103,21 @@ namespace ConcertBrowserUtils
 				.Visibility((InServerFlags & EConcertServerFlags::IgnoreSessionRequirement) != EConcertServerFlags::None ? EVisibility::Visible : EVisibility::Collapsed)
 			];
 	}
+
+	/** Helper function to split an array of FConcertSessionItems into multiple requests and execute*/
+	inline void RequestItemDeletion(IConcertSessionBrowserController& Controller, const TArray<TSharedPtr<FConcertSessionItem>>& SessionItems)
+	{
+		using FServerAdminEndpointId = FGuid;
+		using FSessionId = FGuid;
+			
+		TMap<FServerAdminEndpointId, TArray<FSessionId>> Requests;
+		for (const TSharedPtr<FConcertSessionItem>& Item : SessionItems)
+		{
+			Requests.FindOrAdd(Item->ServerAdminEndpointId).Add(Item->SessionId);
+		}
+		for (auto RequestIt = Requests.CreateConstIterator(); RequestIt; ++RequestIt)
+		{
+			Controller.DeleteSessions(RequestIt->Key, RequestIt->Value);
+		}
+	}
 }
diff --git a/Engine/Plugins/Developer/Concert/ConcertUI/ConcertSharedSlate/Source/ConcertSharedSlate/Public/Session/Browser/IConcertSessionBrowserController.h b/Engine/Plugins/Developer/Concert/ConcertUI/ConcertSharedSlate/Source/ConcertSharedSlate/Public/Session/Browser/IConcertSessionBrowserController.h
--- a/Engine/Plugins/Developer/Concert/ConcertUI/ConcertSharedSlate/Source/ConcertSharedSlate/Public/Session/Browser/IConcertSessionBrowserController.h
+++ b/Engine/Plugins/Developer/Concert/ConcertUI/ConcertSharedSlate/Source/ConcertSharedSlate/Public/Session/Browser/IConcertSessionBrowserController.h
@@ -61,8 +61,7 @@ public:
 	virtual void RenameArchivedSession(const FGuid& ServerAdminEndpointId, const FGuid& SessionId, const FString& NewName) = 0;
 	virtual bool CanRenameActiveSession(const FGuid& ServerAdminEndpointId, const FGuid& SessionId) const = 0;
 	virtual bool CanRenameArchivedSession(const FGuid& ServerAdminEndpointId, const FGuid& SessionId) const = 0;
-	virtual void DeleteActiveSession(const FGuid& ServerAdminEndpointId, const FGuid& SessionId) = 0;
-	virtual void DeleteArchivedSession(const FGuid& ServerAdminEndpointId, const FGuid& SessionId) = 0;
+	virtual void DeleteSessions(const FGuid& ServerAdminEndpointId, const TArray<FGuid>& SessionIds) = 0;
 	virtual bool CanDeleteActiveSession(const FGuid& ServerAdminEndpointId, const FGuid& SessionId) const = 0;
 	virtual bool CanDeleteArchivedSession(const FGuid& ServerAdminEndpointId, const FGuid& SessionId) const = 0;
 	
diff --git a/Engine/Plugins/Developer/Concert/ConcertUI/ConcertSharedSlate/Source/ConcertSharedSlate/Public/Session/Browser/SConcertSessionBrowser.h b/Engine/Plugins/Developer/Concert/ConcertUI/ConcertSharedSlate/Source/ConcertSharedSlate/Public/Session/Browser/SConcertSessionBrowser.h
--- a/Engine/Plugins/Developer/Concert/ConcertUI/ConcertSharedSlate/Source/ConcertSharedSlate/Public/Session/Browser/SConcertSessionBrowser.h
+++ b/Engine/Plugins/Developer/Concert/ConcertUI/ConcertSharedSlate/Source/ConcertSharedSlate/Public/Session/Browser/SConcertSessionBrowser.h
@@ -27,7 +27,8 @@ DECLARE_DELEGATE_RetVal_OneParam(TSharedRef<SWidget>, FExtendSessionTable, const
 DECLARE_DELEGATE_OneParam(FExtenderDelegate, FExtender&);
 DECLARE_DELEGATE_TwoParams(FExtendSessionContextMenu, const TSharedPtr<FConcertSessionItem>&, FExtender&)
 DECLARE_DELEGATE_OneParam(FSessionDelegate, const TSharedPtr<FConcertSessionItem>& /*Session*/);
-DECLARE_DELEGATE_RetVal_OneParam(bool, FCanRemoveSession, const TSharedPtr<FConcertSessionItem>& /* SessionItem */);
+DECLARE_DELEGATE_OneParam(FSessionListDelegate, const TArray<TSharedPtr<FConcertSessionItem>>& /*Sessions*/);
+DECLARE_DELEGATE_RetVal_OneParam(bool, FCanRemoveSessions, const TArray<TSharedPtr<FConcertSessionItem>>& /* SessionItem */);
 
 /**
  * Enables the user to browse/search/filter/sort active and archived sessions, create new session,
@@ -76,12 +77,10 @@ public:
 	/** Called when an archived session is double-clicked. If unset, the default behaviour is to rename the session. */
 	SLATE_EVENT(FSessionDelegate, OnArchivedSessionDoubleClicked)
 	/** Called after a user has requested to delete a session */
-	SLATE_EVENT(FSessionDelegate, OnRequestedDeleteSession)
+	SLATE_EVENT(FSessionListDelegate, PostRequestedDeleteSession)
 
 	/** Ask the user to confirm archiving - most obvious implementation is showing a dialog box */
-	SLATE_EVENT(FCanRemoveSession, CanDeleteArchivedSession)
-	/** Ask the user to confirm deleting - most obvious implementation is showing a dialog box */
-	SLATE_EVENT(FCanRemoveSession, CanDeleteActiveSession)
+	SLATE_EVENT(FCanRemoveSessions, AskUserToDeleteSessions)
 
 	/** Optional snapshot to restore column visibilities with */
 	SLATE_ARGUMENT(FColumnVisibilitySnapshot, ColumnVisibilitySnapshot)
@@ -104,7 +103,7 @@ public:
 	//~ End SWidget Interface
 	
 	void RefreshSessionList();
-	TArray<TSharedPtr<FConcertSessionItem>> GetSessions() const { return Sessions; }
+	bool HasAnySessions() const { return Sessions.Num() > 0; }
 	TArray<TSharedPtr<FConcertSessionItem>> GetSelectedItems() const { return SessionsView->GetSelectedItems(); }
 
 	bool IsNewButtonEnabled() const { return IsNewButtonEnabledInternal(); }
@@ -116,7 +115,7 @@ public:
 	/** Adds row for creating new session. Exposed for other widgets, e.g. discovery overlay to create a new session. */
 	void InsertNewSessionEditableRow() { InsertNewSessionEditableRowInternal(); }
 	/** Creates row under the given (archived) session with which session can be restored. */
-	void InsertRestoreSessionAsEditableRow(const TSharedPtr<FConcertSessionItem>& ArchivedItem) { InsertRestoreSessionAsEditableRowInternal(ArchivedItem); };
+	void InsertRestoreSessionAsEditableRow(const TSharedPtr<FConcertSessionItem>& ArchivedItem) { InsertRestoreSessionAsEditableRowInternal(ArchivedItem); }
 	
 private:
 	
@@ -180,7 +179,7 @@ private:
 	void RequestArchiveSession(const TSharedPtr<FConcertSessionItem>& ActiveItem, const FString& ArchiveName);
 	void RequestRestoreSession(const TSharedPtr<FConcertSessionItem>& RestoreItem, const FString& SessionName);
 	void RequestRenameSession(const TSharedPtr<FConcertSessionItem>& RenamedItem, const FString& NewName);
-	void RequestDeleteSession(const TSharedPtr<FConcertSessionItem>& DeletedItem);
+	void RequestDeleteSessions(const TArray<TSharedPtr<FConcertSessionItem>>& ItemsToDelete);
 
 	TSharedPtr<IConcertSessionBrowserController> GetController() const;
 
@@ -204,9 +203,8 @@ private:
 	/** If unset, the default behaviour is to rename the session. */
 	FSessionDelegate OnArchivedSessionDoubleClicked;
 	/** Called after a live or archived session was requested to be deleted (it may or may not have been deleted).*/
-	FSessionDelegate OnRequestedDeleteSession;
-	FCanRemoveSession CanDeleteArchivedSession;
-	FCanRemoveSession CanDeleteActiveSession;
+	FSessionListDelegate PostRequestedDeleteSession;
+	FCanRemoveSessions AskUserToDeleteSessions;
 
 	// The items displayed in the session list view. It might be filtered and sorted compared to the full list hold by the controller.
 	TArray<TSharedPtr<FConcertSessionItem>> Sessions;
