Make concert UI components more generic

diff --git a/Engine/Plugins/Developer/Concert/ConcertApp/MultiUserServer/Source/MultiUserServer/MultiUserServer.build.cs b/Engine/Plugins/Developer/Concert/ConcertApp/MultiUserServer/Source/MultiUserServer/MultiUserServer.build.cs
--- a/Engine/Plugins/Developer/Concert/ConcertApp/MultiUserServer/Source/MultiUserServer/MultiUserServer.build.cs
+++ b/Engine/Plugins/Developer/Concert/ConcertApp/MultiUserServer/Source/MultiUserServer/MultiUserServer.build.cs
@@ -1,5 +1,7 @@
 // Copyright Epic Games, Inc. All Rights Reserved.
 
+using System.IO;
+
 namespace UnrealBuildTool.Rules
 {
 	public class MultiUserServer : ModuleRules
@@ -33,7 +35,7 @@ namespace UnrealBuildTool.Rules
 					"ConcertSyncServer",
 				}
 			);
-			
+
 			DynamicallyLoadedModuleNames.AddRange(
 				new string[]
 				{
diff --git a/Engine/Plugins/Developer/Concert/ConcertApp/MultiUserServer/Source/MultiUserServer/Private/ModalWindowManager.cpp b/Engine/Plugins/Developer/Concert/ConcertApp/MultiUserServer/Source/MultiUserServer/Private/ModalWindowManager.cpp
deleted file mode 100644
diff --git a/Engine/Plugins/Developer/Concert/ConcertApp/MultiUserServer/Source/MultiUserServer/Private/ModalWindowManager.h b/Engine/Plugins/Developer/Concert/ConcertApp/MultiUserServer/Source/MultiUserServer/Private/ModalWindowManager.h
deleted file mode 100644
diff --git a/Engine/Plugins/Developer/Concert/ConcertApp/MultiUserServer/Source/MultiUserServer/Private/MultiUserServerModule.cpp b/Engine/Plugins/Developer/Concert/ConcertApp/MultiUserServer/Source/MultiUserServer/Private/MultiUserServerModule.cpp
--- a/Engine/Plugins/Developer/Concert/ConcertApp/MultiUserServer/Source/MultiUserServer/Private/MultiUserServerModule.cpp
+++ b/Engine/Plugins/Developer/Concert/ConcertApp/MultiUserServer/Source/MultiUserServer/Private/MultiUserServerModule.cpp
@@ -6,9 +6,8 @@
 #include "ConcertFrontendStyle.h"
 #include "ConcertServerStyle.h"
 #include "ConcertSyncServerLoopInitArgs.h"
-#include "ModalWindowManager.h"
-
-#include "Widgets/ConcertServerWindowController.h"
+#include "Window/ModalWindowManager.h"
+#include "Window/ConcertServerWindowController.h"
 
 #include "Features/IModularFeatures.h"
 #include "Framework/Application/SlateApplication.h"
diff --git a/Engine/Plugins/Developer/Concert/ConcertApp/MultiUserServer/Source/MultiUserServer/Private/Widgets/Browser/ConcertServerSessionBrowserController.cpp b/Engine/Plugins/Developer/Concert/ConcertApp/MultiUserServer/Source/MultiUserServer/Private/Widgets/Browser/ConcertServerSessionBrowserController.cpp
--- a/Engine/Plugins/Developer/Concert/ConcertApp/MultiUserServer/Source/MultiUserServer/Private/Widgets/Browser/ConcertServerSessionBrowserController.cpp
+++ b/Engine/Plugins/Developer/Concert/ConcertApp/MultiUserServer/Source/MultiUserServer/Private/Widgets/Browser/ConcertServerSessionBrowserController.cpp
@@ -8,14 +8,13 @@
 #include "IConcertSyncServer.h"
 #include "Framework/Docking/TabManager.h"
 #include "SConcertServerSessionBrowser.h"
-#include "Framework/Notifications/NotificationManager.h"
 #include "Misc/AsyncTaskNotification.h"
 #include "Session/Browser/ConcertSessionItem.h"
 #include "Textures/SlateIcon.h"
-#include "Widgets/ConcertServerTabs.h"
-#include "Widgets/ConcertServerWindowController.h"
 #include "Widgets/Docking/SDockTab.h"
 #include "Widgets/Notifications/SNotificationList.h"
+#include "Window/ConcertServerTabs.h"
+#include "Window/ConcertServerWindowController.h"
 
 #define LOCTEXT_NAMESPACE "UnrealMultiUserUI"
 
@@ -37,13 +36,20 @@ void FConcertServerSessionBrowserController::Init(const FConcertComponentInitPar
 {
 	ServerInstance = Params.Server;
 	Owner = Params.WindowController;
+	
 	FGlobalTabmanager::Get()->RegisterTabSpawner(
 			ConcertServerTabs::GetSessionBrowserTabId(),
 			FOnSpawnTab::CreateRaw(this, &FConcertServerSessionBrowserController::SpawnSessionBrowserTab)
 		)
 		.SetDisplayName(LOCTEXT("SessionBrowserTabTitle", "Session Browser"))
 		.SetTooltipText(LOCTEXT("SessionBrowserTooltipText", "A section to browse, start, archive, and restore server sessions."))
-		.SetIcon(FSlateIcon(FConcertServerStyle::GetStyleSetName(), TEXT("Concert.MultiUser")));
+		.SetIcon(FSlateIcon(FConcertServerStyle::GetStyleSetName(), TEXT("Concert.MultiUser"))
+	);
+	Params.MainWindowArea->Split(
+		FTabManager::NewStack()
+			->AddTab(ConcertServerTabs::GetSessionBrowserTabId(), ETabState::OpenedTab)
+			->SetForegroundTab(ConcertServerTabs::GetSessionBrowserTabId())
+		);
 
 	ConcertServerEvents::OnLiveSessionCreated().AddSP(this, &FConcertServerSessionBrowserController::OnLiveSessionCreated);
 	ConcertServerEvents::OnLiveSessionDestroyed().AddSP(this, &FConcertServerSessionBrowserController::OnLiveSessionDestroyed);
diff --git a/Engine/Plugins/Developer/Concert/ConcertApp/MultiUserServer/Source/MultiUserServer/Private/Widgets/Browser/ConcertServerSessionBrowserController.h b/Engine/Plugins/Developer/Concert/ConcertApp/MultiUserServer/Source/MultiUserServer/Private/Widgets/Browser/ConcertServerSessionBrowserController.h
--- a/Engine/Plugins/Developer/Concert/ConcertApp/MultiUserServer/Source/MultiUserServer/Private/Widgets/Browser/ConcertServerSessionBrowserController.h
+++ b/Engine/Plugins/Developer/Concert/ConcertApp/MultiUserServer/Source/MultiUserServer/Private/Widgets/Browser/ConcertServerSessionBrowserController.h
@@ -5,7 +5,7 @@
 #include "CoreMinimal.h"
 #include "Misc/AsyncTaskNotification.h"
 #include "Session/Browser/IConcertSessionBrowserController.h"
-#include "Widgets/IConcertComponent.h"
+#include "Window/IConcertComponent.h"
 
 class IConcertServerSession;
 class IConcertServer;
diff --git a/Engine/Plugins/Developer/Concert/ConcertApp/MultiUserServer/Source/MultiUserServer/Private/Widgets/Browser/SConcertServerSessionBrowser.cpp b/Engine/Plugins/Developer/Concert/ConcertApp/MultiUserServer/Source/MultiUserServer/Private/Widgets/Browser/SConcertServerSessionBrowser.cpp
--- a/Engine/Plugins/Developer/Concert/ConcertApp/MultiUserServer/Source/MultiUserServer/Private/Widgets/Browser/SConcertServerSessionBrowser.cpp
+++ b/Engine/Plugins/Developer/Concert/ConcertApp/MultiUserServer/Source/MultiUserServer/Private/Widgets/Browser/SConcertServerSessionBrowser.cpp
@@ -2,19 +2,19 @@
 
 #include "SConcertServerSessionBrowser.h"
 
-#include "ModalWindowManager.h"
 #include "MultiUserServerModule.h"
 #include "MultiUserServerUserSettings.h"
-
-#include "Dialog/SMessageDialog.h"
 #include "Session/Browser/ConcertBrowserUtils.h"
 #include "Session/Browser/ConcertSessionItem.h"
 #include "Session/Browser/SConcertSessionBrowser.h"
 #include "Widgets/Browser/ConcertServerSessionBrowserController.h"
-#include "Widgets/ConcertServerTabs.h"
+#include "Widgets/StatusBar/SConcertStatusBar.h"
+#include "Window/ConcertServerTabs.h"
+#include "Window/ModalWindowManager.h"
+
+#include "Dialog/SMessageDialog.h"
 #include "Widgets/Layout/SBorder.h"
 #include "Widgets/SBoxPanel.h"
-#include "Widgets/StatusBar/SConcertStatusBar.h"
 
 #define LOCTEXT_NAMESPACE "UnrealMultiUserUI"
 
diff --git a/Engine/Plugins/Developer/Concert/ConcertApp/MultiUserServer/Source/MultiUserServer/Private/Widgets/ConcertServerTabs.cpp b/Engine/Plugins/Developer/Concert/ConcertApp/MultiUserServer/Source/MultiUserServer/Private/Widgets/ConcertServerTabs.cpp
deleted file mode 100644
diff --git a/Engine/Plugins/Developer/Concert/ConcertApp/MultiUserServer/Source/MultiUserServer/Private/Widgets/ConcertServerTabs.h b/Engine/Plugins/Developer/Concert/ConcertApp/MultiUserServer/Source/MultiUserServer/Private/Widgets/ConcertServerTabs.h
deleted file mode 100644
diff --git a/Engine/Plugins/Developer/Concert/ConcertApp/MultiUserServer/Source/MultiUserServer/Private/Widgets/ConcertServerWindowController.cpp b/Engine/Plugins/Developer/Concert/ConcertApp/MultiUserServer/Source/MultiUserServer/Private/Widgets/ConcertServerWindowController.cpp
deleted file mode 100644
diff --git a/Engine/Plugins/Developer/Concert/ConcertApp/MultiUserServer/Source/MultiUserServer/Private/Widgets/ConcertServerWindowController.h b/Engine/Plugins/Developer/Concert/ConcertApp/MultiUserServer/Source/MultiUserServer/Private/Widgets/ConcertServerWindowController.h
deleted file mode 100644
diff --git a/Engine/Plugins/Developer/Concert/ConcertApp/MultiUserServer/Source/MultiUserServer/Private/Widgets/IConcertComponent.h b/Engine/Plugins/Developer/Concert/ConcertApp/MultiUserServer/Source/MultiUserServer/Private/Widgets/IConcertComponent.h
deleted file mode 100644
diff --git a/Engine/Plugins/Developer/Concert/ConcertApp/MultiUserServer/Source/MultiUserServer/Private/Widgets/SessionTabs/Archived/ArchivedConcertSessionTab.cpp b/Engine/Plugins/Developer/Concert/ConcertApp/MultiUserServer/Source/MultiUserServer/Private/Widgets/SessionTabs/Archived/ArchivedConcertSessionTab.cpp
--- a/Engine/Plugins/Developer/Concert/ConcertApp/MultiUserServer/Source/MultiUserServer/Private/Widgets/SessionTabs/Archived/ArchivedConcertSessionTab.cpp
+++ b/Engine/Plugins/Developer/Concert/ConcertApp/MultiUserServer/Source/MultiUserServer/Private/Widgets/SessionTabs/Archived/ArchivedConcertSessionTab.cpp
@@ -4,8 +4,8 @@
 
 #include "ArchivedSessionHistoryController.h"
 #include "ConcertLogGlobal.h"
-#include "ModalWindowManager.h"
 #include "MultiUserServerModule.h"
+#include "Window/ModalWindowManager.h"
 
 #include "Algo/AllOf.h"
 #include "Algo/Transform.h"
diff --git a/Engine/Plugins/Developer/Concert/ConcertApp/MultiUserServer/Source/MultiUserServer/Private/Widgets/SessionTabs/ConcertSessionTabBase.cpp b/Engine/Plugins/Developer/Concert/ConcertApp/MultiUserServer/Source/MultiUserServer/Private/Widgets/SessionTabs/ConcertSessionTabBase.cpp
--- a/Engine/Plugins/Developer/Concert/ConcertApp/MultiUserServer/Source/MultiUserServer/Private/Widgets/SessionTabs/ConcertSessionTabBase.cpp
+++ b/Engine/Plugins/Developer/Concert/ConcertApp/MultiUserServer/Source/MultiUserServer/Private/Widgets/SessionTabs/ConcertSessionTabBase.cpp
@@ -4,7 +4,7 @@
 
 #include "IConcertServer.h"
 #include "IConcertSyncServer.h"
-#include "Widgets/ConcertServerTabs.h"
+#include "Window/ConcertServerTabs.h"
 
 #include "Framework/Docking/TabManager.h"
 #include "Widgets/Docking/SDockTab.h"
diff --git a/Engine/Plugins/Developer/Concert/ConcertApp/MultiUserServer/Source/MultiUserServer/Private/Window/ConcertServerTabs.cpp b/Engine/Plugins/Developer/Concert/ConcertApp/MultiUserServer/Source/MultiUserServer/Private/Window/ConcertServerTabs.cpp
new file mode 100644
--- /dev/null
+++ b/Engine/Plugins/Developer/Concert/ConcertApp/MultiUserServer/Source/MultiUserServer/Private/Window/ConcertServerTabs.cpp
@@ -0,0 +1,15 @@
+﻿// Copyright Epic Games, Inc. All Rights Reserved.
+
+#include "ConcertServerTabs.h"
+
+FName ConcertServerTabs::GetSessionBrowserTabId()
+{
+	static FName SessionBrowserTabId("SessionBrowserTabId");
+	return SessionBrowserTabId;
+}
+
+FName ConcertServerTabs::GetClientsTabID()
+{
+	static FName ClientsTabID("ClientsTabID");
+	return ClientsTabID;
+}
\ No newline at end of file
diff --git a/Engine/Plugins/Developer/Concert/ConcertApp/MultiUserServer/Source/MultiUserServer/Private/Window/ConcertServerTabs.h b/Engine/Plugins/Developer/Concert/ConcertApp/MultiUserServer/Source/MultiUserServer/Private/Window/ConcertServerTabs.h
new file mode 100644
--- /dev/null
+++ b/Engine/Plugins/Developer/Concert/ConcertApp/MultiUserServer/Source/MultiUserServer/Private/Window/ConcertServerTabs.h
@@ -0,0 +1,11 @@
+﻿// Copyright Epic Games, Inc. All Rights Reserved.
+
+#pragma once
+
+#include "CoreMinimal.h"
+
+namespace ConcertServerTabs
+{
+	FName GetSessionBrowserTabId();
+	FName GetClientsTabID();
+}
diff --git a/Engine/Plugins/Developer/Concert/ConcertApp/MultiUserServer/Source/MultiUserServer/Private/Window/ConcertServerWindowController.cpp b/Engine/Plugins/Developer/Concert/ConcertApp/MultiUserServer/Source/MultiUserServer/Private/Window/ConcertServerWindowController.cpp
new file mode 100644
--- /dev/null
+++ b/Engine/Plugins/Developer/Concert/ConcertApp/MultiUserServer/Source/MultiUserServer/Private/Window/ConcertServerWindowController.cpp
@@ -0,0 +1,161 @@
+// Copyright Epic Games, Inc. All Rights Reserved.
+
+#include "ConcertServerWindowController.h"
+
+#include "ConcertServerEvents.h"
+#include "IConcertServer.h"
+#include "IConcertSyncServer.h"
+#include "Widgets/Browser//ConcertServerSessionBrowserController.h"
+#include "Widgets/SessionTabs/Archived/ArchivedConcertSessionTab.h"
+#include "Widgets/SessionTabs/Live/LiveConcertSessionTab.h"
+
+#include "Framework/Application/SlateApplication.h"
+#include "Framework/Docking/TabManager.h"
+#include "Framework/Docking/LayoutService.h"
+#include "Framework/Notifications/NotificationManager.h"
+#include "HAL/PlatformApplicationMisc.h"
+#include "Misc/ConfigCacheIni.h"
+#include "OutputLog/Public/OutputLogModule.h"
+
+
+#include "Widgets/SWindow.h"
+
+#define LOCTEXT_NAMESPACE "UnrealMultiUserUI"
+
+FConcertServerWindowController::FConcertServerWindowController(const FConcertServerWindowInitParams& Params)
+	: MultiUserServerLayoutIni(Params.MultiUserServerLayoutIni)
+{
+	ServerInstance = Params.Server;
+	SessionBrowserController = MakeShared<FConcertServerSessionBrowserController>();
+}
+
+FConcertServerWindowController::~FConcertServerWindowController()
+{
+	UnregisterFromSessionDestructionEvents();
+}
+
+TSharedRef<SWindow> FConcertServerWindowController::CreateWindow()
+{
+	FDisplayMetrics DisplayMetrics;
+	FSlateApplication::Get().GetDisplayMetrics(DisplayMetrics);
+	const float DPIScaleFactor = FPlatformApplicationMisc::GetDPIScaleFactorAtPoint(DisplayMetrics.PrimaryDisplayWorkAreaRect.Left, DisplayMetrics.PrimaryDisplayWorkAreaRect.Top);
+
+	const bool bEmbedTitleAreaContent = false;
+	const FVector2D ClientSize(960.0f * DPIScaleFactor, 640.0f * DPIScaleFactor);
+	TSharedRef<SWindow> RootWindowRef = SNew(SWindow)
+		.Title(LOCTEXT("WindowTitle", "Unreal Multi User Server"))
+		.CreateTitleBar(!bEmbedTitleAreaContent)
+		.SupportsMaximize(true)
+		.SupportsMinimize(true)
+		.IsInitiallyMaximized(false)
+		.IsInitiallyMinimized(false)
+		.SizingRule(ESizingRule::UserSized)
+		.AutoCenter(EAutoCenter::PreferredWorkArea)
+		.ClientSize(ClientSize)
+		.AdjustInitialSizeAndPositionForDPIScale(false);
+	RootWindow = RootWindowRef;
+		
+	const bool bShowRootWindowImmediately = false;
+	FSlateApplication::Get().AddWindow(RootWindowRef, bShowRootWindowImmediately);
+	FGlobalTabmanager::Get()->SetRootWindow(RootWindowRef);
+	FGlobalTabmanager::Get()->SetAllowWindowMenuBar(true);
+	FSlateNotificationManager::Get().SetRootWindow(RootWindowRef);
+	
+	TSharedRef<FTabManager::FLayout> DefaultLayout = FTabManager::NewLayout("UnrealMultiUserServerLayout_v1.0");
+	TSharedRef<FTabManager::FArea> MainWindowArea = FTabManager::NewPrimaryArea();
+	InitComponents(MainWindowArea);
+	DefaultLayout->AddArea(MainWindowArea);
+	
+	PersistentLayout = FLayoutSaveRestore::LoadFromConfig(MultiUserServerLayoutIni, DefaultLayout);
+	TSharedPtr<SWidget> Content = FGlobalTabmanager::Get()->RestoreFrom(PersistentLayout.ToSharedRef(), RootWindow, bEmbedTitleAreaContent, EOutputCanBeNullptr::Never);
+	RootWindow->SetContent(Content.ToSharedRef());
+
+	RootWindow->SetOnWindowClosed(FOnWindowClosed::CreateRaw(this, &FConcertServerWindowController::OnWindowClosed));
+	RootWindow->ShowWindow();
+	constexpr bool bForceWindowToFront = true;
+	RootWindow->BringToFront(bForceWindowToFront);
+
+	RegisterForSessionDestructionEvents();
+	return RootWindowRef;
+}
+
+void FConcertServerWindowController::OpenSessionTab(const FGuid& SessionId)
+{
+	if (TSharedPtr<FConcertSessionTabBase> SessionTab = GetOrRegisterSessionTab(SessionId))
+	{
+		SessionTab->OpenSessionTab();
+	}
+}
+
+void FConcertServerWindowController::DestroySessionTab(const FGuid& SessionId)
+{
+	// Destructor will handle the rest, e.g. removing the tab from the window
+	RegisteredSessions.Remove(SessionId);
+}
+
+TSharedPtr<FConcertSessionTabBase> FConcertServerWindowController::GetOrRegisterSessionTab(const FGuid& SessionId)
+{
+	if (const TSharedRef<FConcertSessionTabBase>* FoundId = RegisteredSessions.Find(SessionId))
+	{
+		return *FoundId;
+	}
+	
+	if (TSharedPtr<IConcertServerSession> Session = ServerInstance->GetConcertServer()->GetLiveSession(SessionId))
+	{
+		const TSharedRef<FLiveConcertSessionTab> SessionTab = MakeShared<FLiveConcertSessionTab>(Session.ToSharedRef(), ServerInstance.ToSharedRef(), RootWindow.ToSharedRef());
+		RegisteredSessions.Add(SessionId, SessionTab);
+		return SessionTab;
+	}
+
+	const bool bIsArchivedSession = ServerInstance->GetConcertServer()->GetArchivedSessionInfo(SessionId).IsSet();
+	if (bIsArchivedSession)
+	{
+		const TSharedRef<FArchivedConcertSessionTab> SessionTab = MakeShared<FArchivedConcertSessionTab>(SessionId, ServerInstance.ToSharedRef(), RootWindow.ToSharedRef());
+		RegisteredSessions.Add(SessionId, SessionTab);
+		return SessionTab;
+	}
+	
+	return nullptr;
+}
+
+void FConcertServerWindowController::InitComponents(const TSharedRef<FTabManager::FArea>& MainArea)
+{
+	const FConcertComponentInitParams Params { ServerInstance.ToSharedRef(), SharedThis(this), MainArea };
+	SessionBrowserController->Init(Params);
+}
+
+void FConcertServerWindowController::RegisterForSessionDestructionEvents()
+{
+	ConcertServerEvents::OnLiveSessionDestroyed().AddSP(this, &FConcertServerWindowController::OnLiveSessionDestroyed);
+	ConcertServerEvents::OnArchivedSessionDestroyed().AddSP(this, &FConcertServerWindowController::OnArchivedSessionDestroyed);
+}
+
+void FConcertServerWindowController::UnregisterFromSessionDestructionEvents() const
+{
+	ConcertServerEvents::OnLiveSessionDestroyed().RemoveAll(this);
+	ConcertServerEvents::OnArchivedSessionDestroyed().RemoveAll(this);
+}
+
+void FConcertServerWindowController::OnLiveSessionDestroyed(const IConcertServer&, TSharedRef<IConcertServerSession> InLiveSession)
+{
+	DestroySessionTab(InLiveSession->GetId());
+}
+
+void FConcertServerWindowController::OnArchivedSessionDestroyed(const IConcertServer&, const FGuid& InArchivedSessionId)
+{
+	DestroySessionTab(InArchivedSessionId);
+}
+
+void FConcertServerWindowController::OnWindowClosed(const TSharedRef<SWindow>& Window)
+{
+	SaveLayout();
+	RootWindow.Reset();
+}
+
+void FConcertServerWindowController::SaveLayout() const
+{
+	FLayoutSaveRestore::SaveToConfig(MultiUserServerLayoutIni, PersistentLayout.ToSharedRef());
+    GConfig->Flush(false, MultiUserServerLayoutIni);
+}
+
+#undef LOCTEXT_NAMESPACE 
\ No newline at end of file
diff --git a/Engine/Plugins/Developer/Concert/ConcertApp/MultiUserServer/Source/MultiUserServer/Private/Window/ConcertServerWindowController.h b/Engine/Plugins/Developer/Concert/ConcertApp/MultiUserServer/Source/MultiUserServer/Private/Window/ConcertServerWindowController.h
new file mode 100644
--- /dev/null
+++ b/Engine/Plugins/Developer/Concert/ConcertApp/MultiUserServer/Source/MultiUserServer/Private/Window/ConcertServerWindowController.h
@@ -0,0 +1,78 @@
+// Copyright Epic Games, Inc. All Rights Reserved.
+
+#pragma once
+
+#include "CoreMinimal.h"
+#include "Framework/Docking/TabManager.h"
+
+class IConcertComponent;
+class IConcertServer;
+class FConcertSessionTabBase;
+class FLiveConcertSessionTab;
+class IConcertServerSession;
+class IConcertSyncServer;
+class FConcertServerSessionBrowserController;
+class FOutputLogController;
+class SWindow;
+
+struct FConcertServerWindowInitParams
+{
+	/** The server that the window is supposed to manage */
+	TSharedRef<IConcertSyncServer> Server;
+	
+	/** Config path for server layout ini */
+	FString MultiUserServerLayoutIni;
+	
+	/** Additional modular concert features to enable */
+	TArray<TSharedRef<IConcertComponent>> AdditionalConcertComponents;
+	
+	FConcertServerWindowInitParams(TSharedRef<IConcertSyncServer> Server, FString MultiUserServerLayoutIni = FString())
+		: Server(Server)
+		, MultiUserServerLayoutIni(MoveTemp(MultiUserServerLayoutIni))
+	{}
+};
+
+/** Responsible for creating the Slate window for the server. Implements controller in the model-view-controller pattern. */
+class FConcertServerWindowController : public TSharedFromThis<FConcertServerWindowController>
+{
+public:
+	
+	FConcertServerWindowController(const FConcertServerWindowInitParams& Params);
+	~FConcertServerWindowController();
+	TSharedRef<SWindow> CreateWindow();
+
+	/** Opens or draws attention to the tab for the given live or archived session ID */
+	void OpenSessionTab(const FGuid& SessionId);
+
+	/** Destroys the tab associated with this live or archived session ID. */
+	void DestroySessionTab(const FGuid& SessionId);
+	
+private:
+
+	/** The ini file to use for saving the layout */
+	FString MultiUserServerLayoutIni;
+	/** Holds the current layout for saving later. */
+	TSharedPtr<FTabManager::FLayout> PersistentLayout;
+
+	TSharedPtr<IConcertSyncServer> ServerInstance;
+
+	/** The main window being managed */
+	TSharedPtr<SWindow> RootWindow;
+	TMap<FGuid, TSharedRef<FConcertSessionTabBase>> RegisteredSessions;
+	
+	/** Manages the session browser */
+	TSharedPtr<FConcertServerSessionBrowserController> SessionBrowserController;
+	
+	void InitComponents(const TSharedRef<FTabManager::FArea>& MainArea);
+
+	/** Gets the manager for a session tab if the session ID is valid */
+	TSharedPtr<FConcertSessionTabBase> GetOrRegisterSessionTab(const FGuid& SessionId);
+
+	void RegisterForSessionDestructionEvents();
+	void UnregisterFromSessionDestructionEvents() const;
+	void OnLiveSessionDestroyed(const IConcertServer&, TSharedRef<IConcertServerSession> InLiveSession);
+	void OnArchivedSessionDestroyed(const IConcertServer&, const FGuid& InArchivedSessionId);
+	
+	void OnWindowClosed(const TSharedRef<SWindow>& Window);
+	void SaveLayout() const;
+};
diff --git a/Engine/Plugins/Developer/Concert/ConcertApp/MultiUserServer/Source/MultiUserServer/Private/Window/IConcertComponent.h b/Engine/Plugins/Developer/Concert/ConcertApp/MultiUserServer/Source/MultiUserServer/Private/Window/IConcertComponent.h
new file mode 100644
--- /dev/null
+++ b/Engine/Plugins/Developer/Concert/ConcertApp/MultiUserServer/Source/MultiUserServer/Private/Window/IConcertComponent.h
@@ -0,0 +1,38 @@
+﻿// Copyright Epic Games, Inc. All Rights Reserved.
+
+#pragma once
+
+#include "CoreMinimal.h"
+#include "Framework/Docking/TabManager.h"
+
+class FConcertServerWindowController;
+class IConcertSyncServer;
+
+struct FConcertComponentInitParams
+{
+	/** The server being managed */
+	TSharedRef<IConcertSyncServer> Server;
+
+	/** Manages the server slate application */
+	TSharedRef<FConcertServerWindowController> WindowController;
+
+	/** The root area of the main window layout. Add tabs to this. */
+	TSharedRef<FTabManager::FArea> MainWindowArea;
+
+	FConcertComponentInitParams(TSharedRef<IConcertSyncServer> Server, TSharedRef<FConcertServerWindowController> WindowController, TSharedRef<FTabManager::FArea> MainWindowArea)
+		: Server(MoveTemp(Server))
+		, WindowController(MoveTemp(WindowController))
+		, MainWindowArea(MoveTemp(MainWindowArea))
+	{}
+};
+
+/** Provides the base interface for elements in the concert server UI. */
+class IConcertComponent
+{
+public:
+
+	virtual ~IConcertComponent() = default;
+
+	/** Initialises the component, e.g. registering tab spawners, etc. */
+	virtual void Init(const FConcertComponentInitParams& Params) {}
+};
diff --git a/Engine/Plugins/Developer/Concert/ConcertApp/MultiUserServer/Source/MultiUserServer/Private/Window/ModalWindowManager.cpp b/Engine/Plugins/Developer/Concert/ConcertApp/MultiUserServer/Source/MultiUserServer/Private/Window/ModalWindowManager.cpp
new file mode 100644
--- /dev/null
+++ b/Engine/Plugins/Developer/Concert/ConcertApp/MultiUserServer/Source/MultiUserServer/Private/Window/ModalWindowManager.cpp
@@ -0,0 +1,73 @@
+﻿// Copyright Epic Games, Inc. All Rights Reserved.
+
+#include "ModalWindowManager.h"
+
+#include "Dialog/SCustomDialog.h"
+#include "Widgets/SWindow.h"
+
+UE::MultiUserServer::FModalWindowManager::FModalWindowManager(const TSharedRef<SWindow>& MainWindow)
+{
+	MainWindow->GetOnWindowActivatedEvent().AddRaw(this, &FModalWindowManager::OnMainWindowActivated);
+
+	// Fake modal windows should be closed when the root window is closed
+	MainWindow->GetOnWindowClosedEvent().AddLambda([this](const TSharedRef<SWindow>&)
+	{
+		// One way to close root window while fake modal is open:
+			// 1. Hover MU server window icon (on Windows)
+			// 2. Close root window
+		bRootWindowIsValid = false;
+		
+		for (TWeakPtr<SWindow> ModalWindow : TArray(WindowStack))
+		{
+			if (TSharedPtr<SWindow> Pin = ModalWindow.Pin())
+			{
+				Pin->RequestDestroyWindow();
+			}
+		}
+		WindowStack.Empty();
+	});
+}
+
+void UE::MultiUserServer::FModalWindowManager::ShowFakeModalWindow(const TSharedRef<SCustomDialog>& Dialog)
+{
+	if (ensureMsgf(bRootWindowIsValid, TEXT("Root window has already been destroyed!")))
+	{
+		Dialog->Show();
+		PushWindow(Dialog);
+	}
+}
+
+void UE::MultiUserServer::FModalWindowManager::PushWindow(const TSharedRef<SWindow>& Window)
+{
+	if (ensureMsgf(bRootWindowIsValid, TEXT("Root window has already been destroyed!")))
+	{
+		check(!WindowStack.Contains(Window));
+		WindowStack.Push(Window);
+	
+		Window->GetOnWindowClosedEvent().AddLambda([this](const TSharedRef<SWindow>& Window)
+		{
+			// Might request window destruction twice
+			if (WindowStack.Contains(Window))
+			{
+				PopWindow();
+			}
+		});
+	}
+}
+
+void UE::MultiUserServer::FModalWindowManager::PopWindow()
+{
+	WindowStack.Pop();
+}
+
+void UE::MultiUserServer::FModalWindowManager::OnMainWindowActivated()
+{
+	if (WindowStack.Num() > 0)
+	{
+		if (const TSharedPtr<SWindow> PinnedWindow = WindowStack.Top().Pin())
+		{
+			PinnedWindow->BringToFront();
+			PinnedWindow->FlashWindow();
+		}
+	}
+}
diff --git a/Engine/Plugins/Developer/Concert/ConcertApp/MultiUserServer/Source/MultiUserServer/Private/Window/ModalWindowManager.h b/Engine/Plugins/Developer/Concert/ConcertApp/MultiUserServer/Source/MultiUserServer/Private/Window/ModalWindowManager.h
new file mode 100644
--- /dev/null
+++ b/Engine/Plugins/Developer/Concert/ConcertApp/MultiUserServer/Source/MultiUserServer/Private/Window/ModalWindowManager.h
@@ -0,0 +1,42 @@
+﻿// Copyright Epic Games, Inc. All Rights Reserved.
+
+#pragma once
+
+#include "CoreMinimal.h"
+
+class SWindow;
+class SCustomDialog;
+
+namespace UE::MultiUserServer
+{
+	/**
+	 * Manages "fake" modal windows because the server cannot use "true" modal windows.
+	 *
+	 * True modal windows block the main thread which is not an option because the server must continue responding to
+	 * networked requests.
+	 */
+	class FModalWindowManager
+	{
+	public:
+
+		FModalWindowManager(const TSharedRef<SWindow>& MainWindow);
+
+		void ShowFakeModalWindow(const TSharedRef<SCustomDialog>& Dialog);
+
+		/** Called when a new fake modal window should be tracked on top of all windows. */
+		void PushWindow(const TSharedRef<SWindow>& Window);
+		/** Called when a fake modal window should no longer be tracked, i.e. when it is closed. */
+		void PopWindow();
+
+	private:
+
+		TArray<TWeakPtr<SWindow>> WindowStack;
+
+		/** Modal windows can only be added when this is true */
+		bool bRootWindowIsValid = true;
+
+		void OnMainWindowActivated();
+	};
+}
+
+
