Fix CIS

diff --git a/Engine/Plugins/Developer/Concert/ConcertApp/MultiUserServer/Source/MultiUserServer/Private/Window/ConcertServerWindowController.cpp b/Engine/Plugins/Developer/Concert/ConcertApp/MultiUserServer/Source/MultiUserServer/Private/Window/ConcertServerWindowController.cpp
--- a/Engine/Plugins/Developer/Concert/ConcertApp/MultiUserServer/Source/MultiUserServer/Private/Window/ConcertServerWindowController.cpp
+++ b/Engine/Plugins/Developer/Concert/ConcertApp/MultiUserServer/Source/MultiUserServer/Private/Window/ConcertServerWindowController.cpp
@@ -5,7 +5,7 @@
 #include "ConcertServerEvents.h"
 #include "IConcertServer.h"
 #include "IConcertSyncServer.h"
-#include "Widgets/Browser//ConcertServerSessionBrowserController.h"
+#include "Widgets/Browser/ConcertServerSessionBrowserController.h"
 #include "Widgets/SessionTabs/Archived/ArchivedConcertSessionTab.h"
 #include "Widgets/SessionTabs/Live/LiveConcertSessionTab.h"
 
@@ -24,6 +24,7 @@
 
 FConcertServerWindowController::FConcertServerWindowController(const FConcertServerWindowInitParams& Params)
 	: MultiUserServerLayoutIni(Params.MultiUserServerLayoutIni)
+	, AdditionalConcertComponents(Params.AdditionalConcertComponents)
 {
 	ServerInstance = Params.Server;
 	SessionBrowserController = MakeShared<FConcertServerSessionBrowserController>();
@@ -122,6 +123,11 @@ void FConcertServerWindowController::InitComponents(const TSharedRef<FTabManager
 {
 	const FConcertComponentInitParams Params { ServerInstance.ToSharedRef(), SharedThis(this), MainArea };
 	SessionBrowserController->Init(Params);
+	
+	for (const TSharedRef<IConcertComponent>& ConcertComponent : AdditionalConcertComponents)
+	{
+		ConcertComponent->Init(Params);
+	}
 }
 
 void FConcertServerWindowController::RegisterForSessionDestructionEvents()
diff --git a/Engine/Plugins/Developer/Concert/ConcertApp/MultiUserServer/Source/MultiUserServer/Private/Window/ConcertServerWindowController.h b/Engine/Plugins/Developer/Concert/ConcertApp/MultiUserServer/Source/MultiUserServer/Private/Window/ConcertServerWindowController.h
--- a/Engine/Plugins/Developer/Concert/ConcertApp/MultiUserServer/Source/MultiUserServer/Private/Window/ConcertServerWindowController.h
+++ b/Engine/Plugins/Developer/Concert/ConcertApp/MultiUserServer/Source/MultiUserServer/Private/Window/ConcertServerWindowController.h
@@ -63,6 +63,10 @@ private:
 	/** Manages the session browser */
 	TSharedPtr<FConcertServerSessionBrowserController> SessionBrowserController;
 	
+
+	/** Additional external concert components */
+	TArray<TSharedRef<IConcertComponent>> AdditionalConcertComponents;
+	
 	void InitComponents(const TSharedRef<FTabManager::FArea>& MainArea);
 
 	/** Gets the manager for a session tab if the session ID is valid */
