Refactor tab views in multi user server

diff --git a/Engine/Plugins/Developer/Concert/ConcertApp/MultiUserServer/Source/MultiUserServer/Private/Widgets/Browser/ConcertServerSessionBrowserController.cpp b/Engine/Plugins/Developer/Concert/ConcertApp/MultiUserServer/Source/MultiUserServer/Private/Widgets/Browser/ConcertServerSessionBrowserController.cpp
--- a/Engine/Plugins/Developer/Concert/ConcertApp/MultiUserServer/Source/MultiUserServer/Private/Widgets/Browser/ConcertServerSessionBrowserController.cpp
+++ b/Engine/Plugins/Developer/Concert/ConcertApp/MultiUserServer/Source/MultiUserServer/Private/Widgets/Browser/ConcertServerSessionBrowserController.cpp
@@ -45,11 +45,8 @@ void FConcertServerSessionBrowserController::Init(const FConcertComponentInitPar
 		.SetTooltipText(LOCTEXT("SessionBrowserTooltipText", "A section to browse, start, archive, and restore server sessions."))
 		.SetIcon(FSlateIcon(FConcertServerStyle::GetStyleSetName(), TEXT("Concert.MultiUser"))
 	);
-	Params.MainWindowArea->Split(
-		FTabManager::NewStack()
-			->AddTab(ConcertServerTabs::GetSessionBrowserTabId(), ETabState::OpenedTab)
-			->SetForegroundTab(ConcertServerTabs::GetSessionBrowserTabId())
-		);
+	Params.MainStack->AddTab(ConcertServerTabs::GetSessionBrowserTabId(), ETabState::OpenedTab)
+		->SetForegroundTab(ConcertServerTabs::GetSessionBrowserTabId());
 
 	ConcertServerEvents::OnLiveSessionCreated().AddSP(this, &FConcertServerSessionBrowserController::OnLiveSessionCreated);
 	ConcertServerEvents::OnLiveSessionDestroyed().AddSP(this, &FConcertServerSessionBrowserController::OnLiveSessionDestroyed);
@@ -258,7 +255,6 @@ TSharedRef<SDockTab> FConcertServerSessionBrowserController::SpawnSessionBrowser
 				.DoubleClickArchivedSession(this, &FConcertServerSessionBrowserController::OpenSession)
 		];
 
-	FGlobalTabmanager::Get()->SetMainTab(DockTab);
 	return DockTab;
 }
 
diff --git a/Engine/Plugins/Developer/Concert/ConcertApp/MultiUserServer/Source/MultiUserServer/Private/Widgets/Browser/SConcertServerSessionBrowser.cpp b/Engine/Plugins/Developer/Concert/ConcertApp/MultiUserServer/Source/MultiUserServer/Private/Widgets/Browser/SConcertServerSessionBrowser.cpp
--- a/Engine/Plugins/Developer/Concert/ConcertApp/MultiUserServer/Source/MultiUserServer/Private/Widgets/Browser/SConcertServerSessionBrowser.cpp
+++ b/Engine/Plugins/Developer/Concert/ConcertApp/MultiUserServer/Source/MultiUserServer/Private/Widgets/Browser/SConcertServerSessionBrowser.cpp
@@ -8,44 +8,24 @@
 #include "Session/Browser/ConcertSessionItem.h"
 #include "Session/Browser/SConcertSessionBrowser.h"
 #include "Widgets/Browser/ConcertServerSessionBrowserController.h"
-#include "Widgets/StatusBar/SConcertStatusBar.h"
 #include "Window/ConcertServerTabs.h"
 #include "Window/ModalWindowManager.h"
 
 #include "Dialog/SMessageDialog.h"
-#include "Widgets/Layout/SBorder.h"
-#include "Widgets/SBoxPanel.h"
 
 #define LOCTEXT_NAMESPACE "UnrealMultiUserUI"
 
 void SConcertServerSessionBrowser::Construct(const FArguments& InArgs, TSharedRef<FConcertServerSessionBrowserController> InController)
 {
 	Controller = InController;
-	ChildSlot
-	[
-		SNew(SBorder)
-		.BorderImage(FAppStyle::Get().GetBrush("ToolPanel.GroupBorder"))
-		.Padding(FMargin(1.0f, 2.0f))
+	SConcertTabViewBase::Construct(
+		SConcertTabViewBase::FArguments()
+		.Content()
 		[
-			SNew(SVerticalBox)
-
-			// Session list.
-			+SVerticalBox::Slot()
-			.FillHeight(1.0f)
-			.Padding(1.0f, 2.0f)
-			[
-				MakeSessionTableView(InArgs)
-			]
-
-			+SVerticalBox::Slot()
-			.AutoHeight()
-			.HAlign(HAlign_Fill)
-			.VAlign(VAlign_Bottom)
-			[
-				SNew(SConcertStatusBar, ConcertServerTabs::GetSessionBrowserTabId())
-			]
-		]
-	];
+			MakeSessionTableView(InArgs)
+		],
+		ConcertServerTabs::GetSessionBrowserTabId()
+		);
 }
 
 TSharedRef<SWidget> SConcertServerSessionBrowser::MakeSessionTableView(const FArguments& InArgs)
diff --git a/Engine/Plugins/Developer/Concert/ConcertApp/MultiUserServer/Source/MultiUserServer/Private/Widgets/Browser/SConcertServerSessionBrowser.h b/Engine/Plugins/Developer/Concert/ConcertApp/MultiUserServer/Source/MultiUserServer/Private/Widgets/Browser/SConcertServerSessionBrowser.h
--- a/Engine/Plugins/Developer/Concert/ConcertApp/MultiUserServer/Source/MultiUserServer/Private/Widgets/Browser/SConcertServerSessionBrowser.h
+++ b/Engine/Plugins/Developer/Concert/ConcertApp/MultiUserServer/Source/MultiUserServer/Private/Widgets/Browser/SConcertServerSessionBrowser.h
@@ -5,14 +5,13 @@
 #include "CoreMinimal.h"
 #include "Session/Browser/SConcertSessionBrowser.h"
 #include "Widgets/DeclarativeSyntaxSupport.h"
-#include "Widgets/SCompoundWidget.h"
+#include "Widgets/SConcertTabViewBase.h"
 
-class SMessageDialog;
 class FConcertServerSessionBrowserController;
-class SSearchBox;
+class SConcertSessionBrowser;
 
 /** Shows a list of server sessions */
-class SConcertServerSessionBrowser : public SCompoundWidget
+class SConcertServerSessionBrowser : public SConcertTabViewBase
 {
 public:
 	
diff --git a/Engine/Plugins/Developer/Concert/ConcertApp/MultiUserServer/Source/MultiUserServer/Private/Widgets/SConcertTabViewBase.cpp b/Engine/Plugins/Developer/Concert/ConcertApp/MultiUserServer/Source/MultiUserServer/Private/Widgets/SConcertTabViewBase.cpp
new file mode 100644
--- /dev/null
+++ b/Engine/Plugins/Developer/Concert/ConcertApp/MultiUserServer/Source/MultiUserServer/Private/Widgets/SConcertTabViewBase.cpp
@@ -0,0 +1,38 @@
+// Copyright Epic Games, Inc. All Rights Reserved.
+
+#include "SConcertTabViewBase.h"
+
+#include "Widgets/StatusBar/SConcertStatusBar.h"
+
+#include "Widgets/Layout/SBorder.h"
+#include "Widgets/SBoxPanel.h"
+#include "Styling/AppStyle.h"
+
+void SConcertTabViewBase::Construct(const FArguments& InArgs, FName InStatusBarId)
+{
+	ChildSlot
+	[
+		SNew(SBorder)
+		.BorderImage(FAppStyle::Get().GetBrush("ToolPanel.GroupBorder"))
+		.Padding(FMargin(1.0f, 2.0f))
+		[
+			SNew(SVerticalBox)
+
+			+SVerticalBox::Slot()
+			.FillHeight(1.0f)
+			.Padding(1.0f, 2.0f)
+			[
+				InArgs._Content.Widget
+			]
+
+			+SVerticalBox::Slot()
+			.AutoHeight()
+			.HAlign(HAlign_Fill)
+			.VAlign(VAlign_Bottom)
+			[
+				// This will cause a failing check if InStatusBarId is not unique
+				SNew(SConcertStatusBar, InStatusBarId)
+			]
+		]
+	];
+}
diff --git a/Engine/Plugins/Developer/Concert/ConcertApp/MultiUserServer/Source/MultiUserServer/Private/Widgets/SConcertTabViewBase.h b/Engine/Plugins/Developer/Concert/ConcertApp/MultiUserServer/Source/MultiUserServer/Private/Widgets/SConcertTabViewBase.h
new file mode 100644
--- /dev/null
+++ b/Engine/Plugins/Developer/Concert/ConcertApp/MultiUserServer/Source/MultiUserServer/Private/Widgets/SConcertTabViewBase.h
@@ -0,0 +1,29 @@
+// Copyright Epic Games, Inc. All Rights Reserved.
+
+#pragma once
+
+#include "CoreMinimal.h"
+#include "Widgets/DeclarativeSyntaxSupport.h"
+#include "Widgets/SCompoundWidget.h"
+
+/**
+ * Utility base class for all content that is displayed in an entire major tab.
+ * This class handles shared logic, such as creating a status bar.
+ * 
+ * If the content uses a tab manager, use SConcertTabViewWithManagedBase.
+ */
+class SConcertTabViewBase : public SCompoundWidget
+{
+public:
+
+	SLATE_BEGIN_ARGS(SConcertTabViewBase) {}
+		/** The tab's content */
+		SLATE_NAMED_SLOT(FArguments, Content)
+	SLATE_END_ARGS()
+
+	/**
+	 * @param InArgs
+	 * @param InStatusBarId Unique ID needed for the status bar
+	 */
+	void Construct(const FArguments& InArgs, FName InStatusBarId);
+};
diff --git a/Engine/Plugins/Developer/Concert/ConcertApp/MultiUserServer/Source/MultiUserServer/Private/Widgets/SConcertTabViewWithManagerBase.cpp b/Engine/Plugins/Developer/Concert/ConcertApp/MultiUserServer/Source/MultiUserServer/Private/Widgets/SConcertTabViewWithManagerBase.cpp
new file mode 100644
--- /dev/null
+++ b/Engine/Plugins/Developer/Concert/ConcertApp/MultiUserServer/Source/MultiUserServer/Private/Widgets/SConcertTabViewWithManagerBase.cpp
@@ -0,0 +1,26 @@
+// Copyright Epic Games, Inc. All Rights Reserved.
+
+#include "SConcertTabViewWithManagerBase.h"
+
+#include "Framework/Docking/TabManager.h"
+
+void SConcertTabViewWithManagerBase::Construct(const FArguments& InArgs, FName InStatusBarId)
+{
+	check(InArgs._ConstructUnderWindow && InArgs._ConstructUnderMajorTab);
+	SConcertTabViewBase::Construct(
+		SConcertTabViewBase::FArguments()
+		.Content()
+		[
+			CreateTabs(InArgs)
+		],
+		InStatusBarId
+		);
+}
+
+TSharedRef<SWidget> SConcertTabViewWithManagerBase::CreateTabs(const FArguments& InArgs)
+{
+	TabManager = FGlobalTabmanager::Get()->NewTabManager(InArgs._ConstructUnderMajorTab.ToSharedRef());
+	const TSharedRef<FTabManager::FLayout> Layout = FTabManager::NewLayout(InArgs._LayoutName);
+	InArgs._CreateTabs.ExecuteIfBound(TabManager.ToSharedRef(), Layout);
+	return TabManager->RestoreFrom(Layout, InArgs._ConstructUnderWindow).ToSharedRef();
+}
diff --git a/Engine/Plugins/Developer/Concert/ConcertApp/MultiUserServer/Source/MultiUserServer/Private/Widgets/SConcertTabViewWithManagerBase.h b/Engine/Plugins/Developer/Concert/ConcertApp/MultiUserServer/Source/MultiUserServer/Private/Widgets/SConcertTabViewWithManagerBase.h
new file mode 100644
--- /dev/null
+++ b/Engine/Plugins/Developer/Concert/ConcertApp/MultiUserServer/Source/MultiUserServer/Private/Widgets/SConcertTabViewWithManagerBase.h
@@ -0,0 +1,40 @@
+// Copyright Epic Games, Inc. All Rights Reserved.
+
+#pragma once
+
+#include "CoreMinimal.h"
+#include "Framework/Docking/TabManager.h"
+#include "Widgets/DeclarativeSyntaxSupport.h"
+#include "SConcertTabViewBase.h"
+
+/** Base class for tab views that create sub-tabs. */
+class SConcertTabViewWithManagerBase : public SConcertTabViewBase
+{
+public:
+
+	DECLARE_DELEGATE_TwoParams(FCreateTabs, const TSharedRef<FTabManager>& /*TabManager*/, const TSharedRef<FTabManager::FLayout>& /*Layout*/);
+	
+	SLATE_BEGIN_ARGS(SConcertTabViewWithManagerBase) {}
+		/** Which major tab to construct the sub-tabs under. */
+		SLATE_ARGUMENT(TSharedPtr<SDockTab>, ConstructUnderMajorTab)
+		/** The window in which the sub-tabs will be created. */
+		SLATE_ARGUMENT(TSharedPtr<SWindow>, ConstructUnderWindow)
+		/** Callback for creating the sub-tabs */
+		SLATE_EVENT(FCreateTabs, CreateTabs)
+		/** Name to give the layout. Important for saving config. */
+		SLATE_ARGUMENT(FName, LayoutName)
+	SLATE_END_ARGS()
+
+	/**
+	 * @param InArgs
+	 * @param InStatusBarId Unique ID needed for the status bar
+	 */
+	void Construct(const FArguments& InArgs, FName InStatusBarId);
+
+private:
+
+	/** Holds the tab manager that manages the front-end's tabs. */
+	TSharedPtr<FTabManager> TabManager;
+	
+	TSharedRef<SWidget> CreateTabs(const FArguments& InArgs);
+};
diff --git a/Engine/Plugins/Developer/Concert/ConcertApp/MultiUserServer/Source/MultiUserServer/Private/Widgets/SessionTabs/Archived/ArchivedConcertSessionTab.cpp b/Engine/Plugins/Developer/Concert/ConcertApp/MultiUserServer/Source/MultiUserServer/Private/Widgets/SessionTabs/Archived/ArchivedConcertSessionTab.cpp
--- a/Engine/Plugins/Developer/Concert/ConcertApp/MultiUserServer/Source/MultiUserServer/Private/Widgets/SessionTabs/Archived/ArchivedConcertSessionTab.cpp
+++ b/Engine/Plugins/Developer/Concert/ConcertApp/MultiUserServer/Source/MultiUserServer/Private/Widgets/SessionTabs/Archived/ArchivedConcertSessionTab.cpp
@@ -23,7 +23,7 @@
 #include "Session/History/SSessionHistory.h"
 
 #include "Widgets/HistoryDeletion/SDeleteActivityDependenciesDialog.h"
-#include "Widgets/SessionTabs/Archived/SConcertArchivedSessionInspector.h"
+#include "Widgets/SessionTabs/Archived/SConcertArchivedSessionTabView.h"
 #include "Widgets/StatusBar/SConcertStatusBar.h"
 
 #define LOCTEXT_NAMESPACE "UnrealMultiUserUI"
@@ -45,16 +45,12 @@ void FArchivedConcertSessionTab::CreateDockContent(const TSharedRef<SDockTab>& I
 	});
 	
 	InDockTab->SetContent(
-		SNew(SConcertArchivedSessionInspector)
+		SNew(SConcertArchivedSessionTabView, *GetTabId())
 			.ConstructUnderMajorTab(InDockTab)
 			.ConstructUnderWindow(ConstructUnderWindow.Get())
 			.MakeSessionHistory(MoveTemp(MakeSessionHistory))
 			.DeleteActivity_Raw(this, &FArchivedConcertSessionTab::OnRequestDeleteActivity)
 			.CanDeleteActivity_Raw(this, &FArchivedConcertSessionTab::CanDeleteActivity)
-			.StatusBar()
-			[
-				SNew(SConcertStatusBar, *GetTabId())
-			]
 		);
 }
 
diff --git a/Engine/Plugins/Developer/Concert/ConcertApp/MultiUserServer/Source/MultiUserServer/Private/Widgets/SessionTabs/Archived/ArchivedConcertSessionTab.h b/Engine/Plugins/Developer/Concert/ConcertApp/MultiUserServer/Source/MultiUserServer/Private/Widgets/SessionTabs/Archived/ArchivedConcertSessionTab.h
--- a/Engine/Plugins/Developer/Concert/ConcertApp/MultiUserServer/Source/MultiUserServer/Private/Widgets/SessionTabs/Archived/ArchivedConcertSessionTab.h
+++ b/Engine/Plugins/Developer/Concert/ConcertApp/MultiUserServer/Source/MultiUserServer/Private/Widgets/SessionTabs/Archived/ArchivedConcertSessionTab.h
@@ -8,7 +8,7 @@
 
 struct FConcertSessionActivity;
 class FArchivedSessionHistoryController;
-class SConcertArchivedSessionInspector;
+class SConcertArchivedSessionTabView;
 
 /** Manages the tab for an archived session.  */
 class FArchivedConcertSessionTab : public FConcertSessionTabBase
@@ -38,7 +38,7 @@ private:
 	TSharedPtr<FArchivedSessionHistoryController> HistoryController;
 	
 	/** Displays session */
-	TSharedPtr<SConcertArchivedSessionInspector> Inspector;
+	TSharedPtr<SConcertArchivedSessionTabView> Inspector;
 
 	void OnRequestDeleteActivity(const TSet<TSharedRef<FConcertSessionActivity>>& ActivitiesToDelete) const;
 	FCanDeleteActivitiesResult CanDeleteActivity(const TSet<TSharedRef<FConcertSessionActivity>>& ActivitiesToDelete) const;
diff --git a/Engine/Plugins/Developer/Concert/ConcertApp/MultiUserServer/Source/MultiUserServer/Private/Widgets/SessionTabs/Archived/SConcertArchivedSessionInspector.cpp b/Engine/Plugins/Developer/Concert/ConcertApp/MultiUserServer/Source/MultiUserServer/Private/Widgets/SessionTabs/Archived/SConcertArchivedSessionInspector.cpp
deleted file mode 100644
diff --git a/Engine/Plugins/Developer/Concert/ConcertApp/MultiUserServer/Source/MultiUserServer/Private/Widgets/SessionTabs/Archived/SConcertArchivedSessionInspector.h b/Engine/Plugins/Developer/Concert/ConcertApp/MultiUserServer/Source/MultiUserServer/Private/Widgets/SessionTabs/Archived/SConcertArchivedSessionInspector.h
deleted file mode 100644
diff --git a/Engine/Plugins/Developer/Concert/ConcertApp/MultiUserServer/Source/MultiUserServer/Private/Widgets/SessionTabs/Archived/SConcertArchivedSessionTabView.cpp b/Engine/Plugins/Developer/Concert/ConcertApp/MultiUserServer/Source/MultiUserServer/Private/Widgets/SessionTabs/Archived/SConcertArchivedSessionTabView.cpp
new file mode 100644
--- /dev/null
+++ b/Engine/Plugins/Developer/Concert/ConcertApp/MultiUserServer/Source/MultiUserServer/Private/Widgets/SessionTabs/Archived/SConcertArchivedSessionTabView.cpp
@@ -0,0 +1,64 @@
+﻿// Copyright Epic Games, Inc. All Rights Reserved.
+
+#include "SConcertArchivedSessionTabView.h"
+
+#include "Framework/Docking/TabManager.h"
+#include "Session/History/SSessionHistory.h"
+#include "Widgets/Docking/SDockTab.h"
+
+#define LOCTEXT_NAMESPACE "SConcertArchivedSessionInspector"
+
+const FName SConcertArchivedSessionTabView::HistoryTabId("HistoryTabId");
+
+void SConcertArchivedSessionTabView::Construct(const FArguments& InArgs, FName InStatusBarID)
+{
+	check(InArgs._MakeSessionHistory.IsBound() && InArgs._CanDeleteActivity.IsBound());
+	SConcertTabViewWithManagerBase::Construct(
+		SConcertTabViewWithManagerBase::FArguments()
+		.ConstructUnderWindow(InArgs._ConstructUnderWindow)
+		.ConstructUnderMajorTab(InArgs._ConstructUnderMajorTab)
+		.CreateTabs(FCreateTabs::CreateLambda([this, &InArgs](const TSharedRef<FTabManager>& InTabManager, const TSharedRef<FTabManager::FLayout>& InLayout)
+		{
+			CreateTabs(InTabManager, InLayout, InArgs);
+		}))
+		.LayoutName("ConcertArchivedSessionInspector_v0.1"),
+		InStatusBarID
+	);
+}
+
+void SConcertArchivedSessionTabView::CreateTabs(const TSharedRef<FTabManager>& InTabManager, const TSharedRef<FTabManager::FLayout>& InLayout, const FArguments& InArgs)
+{
+	InTabManager->RegisterTabSpawner(HistoryTabId, FOnSpawnTab::CreateSP(this, &SConcertArchivedSessionTabView::SpawnActivityHistory, InArgs._MakeSessionHistory, InArgs._CanDeleteActivity, InArgs._DeleteActivity))
+		.SetDisplayName(LOCTEXT("ActivityHistoryLabel", "History"));
+	InLayout->AddArea
+		(
+			FTabManager::NewPrimaryArea()
+				->SetOrientation(Orient_Horizontal)
+				->Split
+				(
+					FTabManager::NewStack()
+					->SetSizeCoefficient(0.5f)
+					->AddTab(HistoryTabId, ETabState::OpenedTab)
+				)
+		);
+}
+
+TSharedRef<SDockTab> SConcertArchivedSessionTabView::SpawnActivityHistory(
+	const FSpawnTabArgs& Args,
+	SEditableSessionHistory::FMakeSessionHistory MakeSessionHistory,
+	SEditableSessionHistory::FCanDeleteActivities CanDeleteActivity,
+	SEditableSessionHistory::FRequestDeleteActivities DeleteActivity)
+{
+	SessionHistory = SNew(SEditableSessionHistory)
+		.MakeSessionHistory(MoveTemp(MakeSessionHistory))
+		.CanDeleteActivity(MoveTemp(CanDeleteActivity))
+		.DeleteActivity(MoveTemp(DeleteActivity));
+	return SNew(SDockTab)
+		.Label(LOCTEXT("ActivityHistoryLabel", "History"))
+		.TabRole(PanelTab)
+		[
+			SessionHistory.ToSharedRef()
+		]; 
+}
+
+#undef LOCTEXT_NAMESPACE
diff --git a/Engine/Plugins/Developer/Concert/ConcertApp/MultiUserServer/Source/MultiUserServer/Private/Widgets/SessionTabs/Archived/SConcertArchivedSessionTabView.h b/Engine/Plugins/Developer/Concert/ConcertApp/MultiUserServer/Source/MultiUserServer/Private/Widgets/SessionTabs/Archived/SConcertArchivedSessionTabView.h
new file mode 100644
--- /dev/null
+++ b/Engine/Plugins/Developer/Concert/ConcertApp/MultiUserServer/Source/MultiUserServer/Private/Widgets/SessionTabs/Archived/SConcertArchivedSessionTabView.h
@@ -0,0 +1,44 @@
+﻿// Copyright Epic Games, Inc. All Rights Reserved.
+
+#pragma once
+
+#include "CoreMinimal.h"
+
+#include "Session/History/SEditableSessionHistory.h"
+#include "Session/History/SSessionHistory.h"
+#include "Widgets/DeclarativeSyntaxSupport.h"
+#include "Framework/Docking/TabManager.h"
+#include "Widgets/SConcertTabViewWithManagerBase.h"
+
+class SDockTab;
+class SSessionHistory;
+
+/** Designed as content for a tab. Displays information about an archived session. */
+class SConcertArchivedSessionTabView : public SConcertTabViewWithManagerBase
+{
+public:
+	
+	static const FName HistoryTabId;
+
+	SLATE_BEGIN_ARGS(SConcertArchivedSessionTabView)
+	{}
+		SLATE_ARGUMENT(TSharedPtr<SDockTab>, ConstructUnderMajorTab)
+		SLATE_ARGUMENT(TSharedPtr<SWindow>, ConstructUnderWindow)
+		SLATE_EVENT(SEditableSessionHistory::FMakeSessionHistory, MakeSessionHistory)
+		SLATE_EVENT(SEditableSessionHistory::FCanDeleteActivities, CanDeleteActivity)
+		SLATE_EVENT(SEditableSessionHistory::FRequestDeleteActivities, DeleteActivity)
+	SLATE_END_ARGS()
+
+	void Construct(const FArguments& InArgs, FName InStatusBarID);
+
+private:
+
+	TSharedPtr<SEditableSessionHistory> SessionHistory;
+	
+	void CreateTabs(const TSharedRef<FTabManager>& InTabManager, const TSharedRef<FTabManager::FLayout>& InLayout, const FArguments& InArgs);
+	TSharedRef<SDockTab> SpawnActivityHistory(const FSpawnTabArgs& Args,
+		SEditableSessionHistory::FMakeSessionHistory FMakeSessionHistory,
+		SEditableSessionHistory::FCanDeleteActivities CanDeleteActivity,
+		SEditableSessionHistory::FRequestDeleteActivities DeleteActivity
+		);
+};
diff --git a/Engine/Plugins/Developer/Concert/ConcertApp/MultiUserServer/Source/MultiUserServer/Private/Widgets/SessionTabs/Live/LiveConcertSessionTab.cpp b/Engine/Plugins/Developer/Concert/ConcertApp/MultiUserServer/Source/MultiUserServer/Private/Widgets/SessionTabs/Live/LiveConcertSessionTab.cpp
--- a/Engine/Plugins/Developer/Concert/ConcertApp/MultiUserServer/Source/MultiUserServer/Private/Widgets/SessionTabs/Live/LiveConcertSessionTab.cpp
+++ b/Engine/Plugins/Developer/Concert/ConcertApp/MultiUserServer/Source/MultiUserServer/Private/Widgets/SessionTabs/Live/LiveConcertSessionTab.cpp
@@ -5,7 +5,7 @@
 #include "IConcertSession.h"
 #include "LiveServerSessionHistoryController.h"
 #include "PackageViewer/ConcertSessionPackageViewerController.h"
-#include "SConcertSessionInspector.h"
+#include "SConcertLiveSessionTabView.h"
 #include "Widgets/StatusBar/SConcertStatusBar.h"
 
 #include "Widgets/Docking/SDockTab.h"
@@ -20,7 +20,7 @@ FLiveConcertSessionTab::FLiveConcertSessionTab(TSharedRef<IConcertServerSession>
 
 void FLiveConcertSessionTab::CreateDockContent(const TSharedRef<SDockTab>& InDockTab)
 {
-	const SConcertSessionInspector::FRequiredArgs WidgetArgs
+	const SConcertLiveSessionTabView::FRequiredWidgets WidgetArgs
 	{
 		InDockTab,
 		ConstructUnderWindow.Get(),
@@ -28,11 +28,7 @@ void FLiveConcertSessionTab::CreateDockContent(const TSharedRef<SDockTab>& InDoc
 		PackageViewerController->GetPackageViewer()
 	};
 	InDockTab->SetContent(
-		SNew(SConcertSessionInspector, WidgetArgs)
-			.StatusBar()
-			[
-				SNew(SConcertStatusBar, *GetTabId())
-			]
+		SNew(SConcertLiveSessionTabView, WidgetArgs, *GetTabId())
 		);
 }
 
diff --git a/Engine/Plugins/Developer/Concert/ConcertApp/MultiUserServer/Source/MultiUserServer/Private/Widgets/SessionTabs/Live/SConcertLiveSessionTabView.cpp b/Engine/Plugins/Developer/Concert/ConcertApp/MultiUserServer/Source/MultiUserServer/Private/Widgets/SessionTabs/Live/SConcertLiveSessionTabView.cpp
new file mode 100644
--- /dev/null
+++ b/Engine/Plugins/Developer/Concert/ConcertApp/MultiUserServer/Source/MultiUserServer/Private/Widgets/SessionTabs/Live/SConcertLiveSessionTabView.cpp
@@ -0,0 +1,106 @@
+﻿// Copyright Epic Games, Inc. All Rights Reserved.
+
+#include "SConcertLiveSessionTabView.h"
+
+#include "Framework/Docking/TabManager.h"
+#include "Session/History/SSessionHistory.h"
+#include "Widgets/Docking/SDockTab.h"
+#include "Widgets/Layout/SBorder.h"
+#include "Widgets/SBoxPanel.h"
+
+#define LOCTEXT_NAMESPACE "SConcertSessionInspector"
+
+const FName SConcertLiveSessionTabView::HistoryTabId("HistoryTabId");
+const FName SConcertLiveSessionTabView::SessionContentTabId("SessionContentTabId");
+const FName SConcertLiveSessionTabView::ConnectionMonitorTabId("ConnectionMonitorTabId");
+
+void SConcertLiveSessionTabView::Construct(const FArguments& InArgs, const FRequiredWidgets& InRequiredArgs, FName StatusBarId)
+{
+	SConcertTabViewWithManagerBase::Construct(
+		SConcertTabViewWithManagerBase::FArguments()
+		.ConstructUnderWindow(InRequiredArgs.ConstructUnderWindow)
+		.ConstructUnderMajorTab(InRequiredArgs.ConstructUnderMajorTab)
+		.CreateTabs(FCreateTabs::CreateLambda([this, &InRequiredArgs](const TSharedRef<FTabManager>& InTabManager, const TSharedRef<FTabManager::FLayout>& InLayout)
+		{
+			CreateTabs(InTabManager, InLayout, InRequiredArgs);
+		}))
+		.LayoutName("ConcertSessionInspector_v0.1"),
+		StatusBarId
+		);
+}
+
+TSharedRef<SWidget> SConcertLiveSessionTabView::CreateTabs(const TSharedRef<FTabManager>& InTabManager, const TSharedRef<FTabManager::FLayout>& InLayout, const FRequiredWidgets& InRequiredArgs)
+{
+	InTabManager->RegisterTabSpawner(HistoryTabId, FOnSpawnTab::CreateSP(this, &SConcertLiveSessionTabView::SpawnActivityHistory, InRequiredArgs.SessionHistory))
+		.SetDisplayName(LOCTEXT("ActivityHistoryLabel", "History"));
+	
+	InTabManager->RegisterTabSpawner(SessionContentTabId, FOnSpawnTab::CreateSP(this, &SConcertLiveSessionTabView::SpawnSessionContent, InRequiredArgs.PackageViewer))
+		.SetDisplayName(LOCTEXT("SessionContentLabel", "Session Content"));
+	
+	InTabManager->RegisterTabSpawner(ConnectionMonitorTabId, FOnSpawnTab::CreateSP(this, &SConcertLiveSessionTabView::SpawnConnectionMonitor))
+		.SetDisplayName(LOCTEXT("ConnectionMonitorLabel", "Connection Monitor"));
+
+	InLayout->AddArea
+	(
+		FTabManager::NewPrimaryArea()
+			->SetOrientation(Orient_Horizontal)
+			->Split
+			(
+				FTabManager::NewStack()
+				->SetSizeCoefficient(0.5f)
+				->AddTab(HistoryTabId, ETabState::OpenedTab)
+			)
+			->Split
+			(
+				FTabManager::NewSplitter()
+				->SetSizeCoefficient(0.5f)
+				->SetOrientation(Orient_Vertical)
+				->Split
+				(
+					FTabManager::NewStack()
+					->SetSizeCoefficient(0.5f)
+					->AddTab(SessionContentTabId, ETabState::OpenedTab)
+				)
+				->Split
+				(
+					FTabManager::NewStack()
+					->SetSizeCoefficient(0.5f)
+					->AddTab(ConnectionMonitorTabId, ETabState::OpenedTab)
+				)
+			)
+	);
+
+	return InTabManager->RestoreFrom(InLayout, InRequiredArgs.ConstructUnderWindow).ToSharedRef();
+}
+
+TSharedRef<SDockTab> SConcertLiveSessionTabView::SpawnActivityHistory(const FSpawnTabArgs& Args, TSharedRef<SSessionHistory> SessionHistory)
+{
+	return SNew(SDockTab)
+		.Label(LOCTEXT("ActivityHistoryLabel", "History"))
+		.TabRole(PanelTab)
+		[
+			SessionHistory
+		];
+}
+
+TSharedRef<SDockTab> SConcertLiveSessionTabView::SpawnSessionContent(const FSpawnTabArgs& Args, TSharedRef<SConcertSessionPackageViewer> PackageViewer)
+{
+	return SNew(SDockTab)
+		.Label(LOCTEXT("SessionContentLabel", "Session Content"))
+		.TabRole(PanelTab)
+		[
+			PackageViewer
+		];
+}
+
+TSharedRef<SDockTab> SConcertLiveSessionTabView::SpawnConnectionMonitor(const FSpawnTabArgs& Args)
+{
+	return SNew(SDockTab)
+		.Label(LOCTEXT("ConnectionMonitorLabel", "Connection Monitor"))
+		.TabRole(PanelTab)
+		[
+			SNullWidget::NullWidget
+		];
+}
+
+#undef LOCTEXT_NAMESPACE
diff --git a/Engine/Plugins/Developer/Concert/ConcertApp/MultiUserServer/Source/MultiUserServer/Private/Widgets/SessionTabs/Live/SConcertLiveSessionTabView.h b/Engine/Plugins/Developer/Concert/ConcertApp/MultiUserServer/Source/MultiUserServer/Private/Widgets/SessionTabs/Live/SConcertLiveSessionTabView.h
new file mode 100644
--- /dev/null
+++ b/Engine/Plugins/Developer/Concert/ConcertApp/MultiUserServer/Source/MultiUserServer/Private/Widgets/SessionTabs/Live/SConcertLiveSessionTabView.h
@@ -0,0 +1,62 @@
+﻿// Copyright Epic Games, Inc. All Rights Reserved.
+
+#pragma once
+
+#include "CoreMinimal.h"
+#include "Session/History/SSessionHistory.h"
+#include "PackageViewer/SConcertSessionPackageViewer.h"
+#include "Widgets/DeclarativeSyntaxSupport.h"
+#include "Widgets/SConcertTabViewWithManagerBase.h"
+
+class FTabManager;
+class FSpawnTabArgs;
+class SDockTab;
+class SWindow;
+
+/**
+ * Designed to be the content of a tab showing:
+ *  - activity history (transactions stored on the server as well as who made those transactions)
+ *  - session content (list of session data saved during a Multi-user session)
+ *  - connection monitor (details about the connected clients on the given session and network info)
+ * Implements view in model-view-controller pattern.
+ */
+class SConcertLiveSessionTabView : public SConcertTabViewWithManagerBase
+{
+public:
+
+	struct FRequiredWidgets
+	{
+		TSharedRef<SDockTab> ConstructUnderMajorTab;
+		TSharedRef<SWindow> ConstructUnderWindow;
+		TSharedRef<SSessionHistory> SessionHistory;
+		TSharedRef<SConcertSessionPackageViewer> PackageViewer;
+
+		FRequiredWidgets(TSharedRef<SDockTab> ConstructUnderMajorTab,
+				TSharedRef<SWindow> ConstructUnderWindow,
+				TSharedRef<SSessionHistory> SessionHistoryController,
+				TSharedRef<SConcertSessionPackageViewer> PackageViewerController)
+			: ConstructUnderMajorTab(MoveTemp(ConstructUnderMajorTab))
+			, ConstructUnderWindow(MoveTemp(ConstructUnderWindow))
+			, SessionHistory(MoveTemp(SessionHistoryController))
+			, PackageViewer(MoveTemp(PackageViewerController))
+		{}
+	};
+
+	static const FName HistoryTabId;
+	static const FName SessionContentTabId;
+	static const FName ConnectionMonitorTabId;
+
+	SLATE_BEGIN_ARGS(SConcertLiveSessionTabView) {}
+		SLATE_NAMED_SLOT(FArguments, StatusBar)
+	SLATE_END_ARGS()
+
+	void Construct(const FArguments& InArgs, const FRequiredWidgets& InRequiredArgs, FName StatusBarId);
+
+private:
+
+	// Spawn tabs
+	TSharedRef<SWidget> CreateTabs(const TSharedRef<FTabManager>& InTabManager, const TSharedRef<FTabManager::FLayout>& InLayout, const FRequiredWidgets& InRequiredArgs);
+	TSharedRef<SDockTab> SpawnActivityHistory(const FSpawnTabArgs& Args, TSharedRef<SSessionHistory> SessionHistory);
+	TSharedRef<SDockTab> SpawnSessionContent(const FSpawnTabArgs& Args, TSharedRef<SConcertSessionPackageViewer> PackageViewer);
+	TSharedRef<SDockTab> SpawnConnectionMonitor(const FSpawnTabArgs& Args);
+};
diff --git a/Engine/Plugins/Developer/Concert/ConcertApp/MultiUserServer/Source/MultiUserServer/Private/Widgets/SessionTabs/Live/SConcertSessionInspector.cpp b/Engine/Plugins/Developer/Concert/ConcertApp/MultiUserServer/Source/MultiUserServer/Private/Widgets/SessionTabs/Live/SConcertSessionInspector.cpp
deleted file mode 100644
diff --git a/Engine/Plugins/Developer/Concert/ConcertApp/MultiUserServer/Source/MultiUserServer/Private/Widgets/SessionTabs/Live/SConcertSessionInspector.h b/Engine/Plugins/Developer/Concert/ConcertApp/MultiUserServer/Source/MultiUserServer/Private/Widgets/SessionTabs/Live/SConcertSessionInspector.h
deleted file mode 100644
diff --git a/Engine/Plugins/Developer/Concert/ConcertApp/MultiUserServer/Source/MultiUserServer/Private/Widgets/StatusBar/SConcertStatusBar.cpp b/Engine/Plugins/Developer/Concert/ConcertApp/MultiUserServer/Source/MultiUserServer/Private/Widgets/StatusBar/SConcertStatusBar.cpp
--- a/Engine/Plugins/Developer/Concert/ConcertApp/MultiUserServer/Source/MultiUserServer/Private/Widgets/StatusBar/SConcertStatusBar.cpp
+++ b/Engine/Plugins/Developer/Concert/ConcertApp/MultiUserServer/Source/MultiUserServer/Private/Widgets/StatusBar/SConcertStatusBar.cpp
@@ -85,8 +85,15 @@ namespace UE::ConcertServerUI::Private
 				// Destroying StatusBarOutputLog in ~FStatusBarSingleton is too late: it causes a crash
 				FSlateApplication::Get().OnPreShutdown().AddRaw(this, &FStatusBarSingleton::PreShutdownSlate);
 			}
+
+			const bool bIsDrawerNameUnique = !StatusBars.ContainsByPredicate([&WidgetDrawer](TWeakPtr<SWidgetDrawer> WeakDrawer)
+			{
+				return ensure(WeakDrawer.IsValid())
+					&& WeakDrawer.Pin()->GetDrawerName() == WidgetDrawer->GetDrawerName();
+			});
+			checkf(bIsDrawerNameUnique, TEXT("Every widget drawer is expected to have an unique ID"));
 			
-			StatusBars.Add(WidgetDrawer);
+			StatusBars.Add(MoveTemp(WidgetDrawer));
 			
 			OutputLogDrawer.GetDrawerContentDelegate.BindRaw(this, &FStatusBarSingleton::OnGetOutputLog);
 			OutputLogDrawer.OnDrawerOpenedDelegate.BindRaw(this, &FStatusBarSingleton::OnOutputLogOpened);
diff --git a/Engine/Plugins/Developer/Concert/ConcertApp/MultiUserServer/Source/MultiUserServer/Private/Window/ConcertServerWindowController.cpp b/Engine/Plugins/Developer/Concert/ConcertApp/MultiUserServer/Source/MultiUserServer/Private/Window/ConcertServerWindowController.cpp
--- a/Engine/Plugins/Developer/Concert/ConcertApp/MultiUserServer/Source/MultiUserServer/Private/Window/ConcertServerWindowController.cpp
+++ b/Engine/Plugins/Developer/Concert/ConcertApp/MultiUserServer/Source/MultiUserServer/Private/Window/ConcertServerWindowController.cpp
@@ -24,10 +24,11 @@
 
 FConcertServerWindowController::FConcertServerWindowController(const FConcertServerWindowInitParams& Params)
 	: MultiUserServerLayoutIni(Params.MultiUserServerLayoutIni)
-	, AdditionalConcertComponents(Params.AdditionalConcertComponents)
+	, ServerInstance(Params.Server)
+	, SessionBrowserController(MakeShared<FConcertServerSessionBrowserController>())
+	, ConcertComponents(Params.AdditionalConcertComponents)
 {
-	ServerInstance = Params.Server;
-	SessionBrowserController = MakeShared<FConcertServerSessionBrowserController>();
+	ConcertComponents.Add(SessionBrowserController);
 }
 
 FConcertServerWindowController::~FConcertServerWindowController()
@@ -62,9 +63,11 @@ TSharedRef<SWindow> FConcertServerWindowController::CreateWindow()
 	FGlobalTabmanager::Get()->SetAllowWindowMenuBar(true);
 	FSlateNotificationManager::Get().SetRootWindow(RootWindowRef);
 	
-	TSharedRef<FTabManager::FLayout> DefaultLayout = FTabManager::NewLayout("UnrealMultiUserServerLayout_v1.0");
-	TSharedRef<FTabManager::FArea> MainWindowArea = FTabManager::NewPrimaryArea();
-	InitComponents(MainWindowArea);
+	const TSharedRef<FTabManager::FStack> MainStack = FTabManager::NewStack();
+	InitComponents(MainStack);
+	const TSharedRef<FTabManager::FArea> MainWindowArea = FTabManager::NewPrimaryArea();
+	const TSharedRef<FTabManager::FLayout> DefaultLayout = FTabManager::NewLayout("UnrealMultiUserServerLayout_v1.0");
+	MainWindowArea->Split(MainStack);
 	DefaultLayout->AddArea(MainWindowArea);
 	
 	PersistentLayout = FLayoutSaveRestore::LoadFromConfig(MultiUserServerLayoutIni, DefaultLayout);
@@ -82,7 +85,7 @@ TSharedRef<SWindow> FConcertServerWindowController::CreateWindow()
 
 void FConcertServerWindowController::OpenSessionTab(const FGuid& SessionId)
 {
-	if (TSharedPtr<FConcertSessionTabBase> SessionTab = GetOrRegisterSessionTab(SessionId))
+	if (const TSharedPtr<FConcertSessionTabBase> SessionTab = GetOrRegisterSessionTab(SessionId))
 	{
 		SessionTab->OpenSessionTab();
 	}
@@ -103,7 +106,7 @@ TSharedPtr<FConcertSessionTabBase> FConcertServerWindowController::GetOrRegister
 	
 	if (TSharedPtr<IConcertServerSession> Session = ServerInstance->GetConcertServer()->GetLiveSession(SessionId))
 	{
-		const TSharedRef<FLiveConcertSessionTab> SessionTab = MakeShared<FLiveConcertSessionTab>(Session.ToSharedRef(), ServerInstance.ToSharedRef(), RootWindow.ToSharedRef());
+		const TSharedRef<FLiveConcertSessionTab> SessionTab = MakeShared<FLiveConcertSessionTab>(Session.ToSharedRef(), ServerInstance, RootWindow.ToSharedRef());
 		RegisteredSessions.Add(SessionId, SessionTab);
 		return SessionTab;
 	}
@@ -111,7 +114,7 @@ TSharedPtr<FConcertSessionTabBase> FConcertServerWindowController::GetOrRegister
 	const bool bIsArchivedSession = ServerInstance->GetConcertServer()->GetArchivedSessionInfo(SessionId).IsSet();
 	if (bIsArchivedSession)
 	{
-		const TSharedRef<FArchivedConcertSessionTab> SessionTab = MakeShared<FArchivedConcertSessionTab>(SessionId, ServerInstance.ToSharedRef(), RootWindow.ToSharedRef());
+		const TSharedRef<FArchivedConcertSessionTab> SessionTab = MakeShared<FArchivedConcertSessionTab>(SessionId, ServerInstance, RootWindow.ToSharedRef());
 		RegisteredSessions.Add(SessionId, SessionTab);
 		return SessionTab;
 	}
@@ -119,12 +122,10 @@ TSharedPtr<FConcertSessionTabBase> FConcertServerWindowController::GetOrRegister
 	return nullptr;
 }
 
-void FConcertServerWindowController::InitComponents(const TSharedRef<FTabManager::FArea>& MainArea)
+void FConcertServerWindowController::InitComponents(const TSharedRef<FTabManager::FStack>& MainArea)
 {
-	const FConcertComponentInitParams Params { ServerInstance.ToSharedRef(), SharedThis(this), MainArea };
-	SessionBrowserController->Init(Params);
-	
-	for (const TSharedRef<IConcertComponent>& ConcertComponent : AdditionalConcertComponents)
+	const FConcertComponentInitParams Params { ServerInstance, SharedThis(this), MainArea };
+	for (const TSharedRef<IConcertComponent>& ConcertComponent : ConcertComponents)
 	{
 		ConcertComponent->Init(Params);
 	}
diff --git a/Engine/Plugins/Developer/Concert/ConcertApp/MultiUserServer/Source/MultiUserServer/Private/Window/ConcertServerWindowController.h b/Engine/Plugins/Developer/Concert/ConcertApp/MultiUserServer/Source/MultiUserServer/Private/Window/ConcertServerWindowController.h
--- a/Engine/Plugins/Developer/Concert/ConcertApp/MultiUserServer/Source/MultiUserServer/Private/Window/ConcertServerWindowController.h
+++ b/Engine/Plugins/Developer/Concert/ConcertApp/MultiUserServer/Source/MultiUserServer/Private/Window/ConcertServerWindowController.h
@@ -54,20 +54,18 @@ private:
 	/** Holds the current layout for saving later. */
 	TSharedPtr<FTabManager::FLayout> PersistentLayout;
 
-	TSharedPtr<IConcertSyncServer> ServerInstance;
+	TSharedRef<IConcertSyncServer> ServerInstance;
 
 	/** The main window being managed */
 	TSharedPtr<SWindow> RootWindow;
 	TMap<FGuid, TSharedRef<FConcertSessionTabBase>> RegisteredSessions;
 	
 	/** Manages the session browser */
-	TSharedPtr<FConcertServerSessionBrowserController> SessionBrowserController;
-	
-
+	TSharedRef<FConcertServerSessionBrowserController> SessionBrowserController;
 	/** Additional external concert components */
-	TArray<TSharedRef<IConcertComponent>> AdditionalConcertComponents;
+	TArray<TSharedRef<IConcertComponent>> ConcertComponents;
 	
-	void InitComponents(const TSharedRef<FTabManager::FArea>& MainArea);
+	void InitComponents(const TSharedRef<FTabManager::FStack>& MainArea);
 
 	/** Gets the manager for a session tab if the session ID is valid */
 	TSharedPtr<FConcertSessionTabBase> GetOrRegisterSessionTab(const FGuid& SessionId);
diff --git a/Engine/Plugins/Developer/Concert/ConcertApp/MultiUserServer/Source/MultiUserServer/Private/Window/IConcertComponent.h b/Engine/Plugins/Developer/Concert/ConcertApp/MultiUserServer/Source/MultiUserServer/Private/Window/IConcertComponent.h
--- a/Engine/Plugins/Developer/Concert/ConcertApp/MultiUserServer/Source/MultiUserServer/Private/Window/IConcertComponent.h
+++ b/Engine/Plugins/Developer/Concert/ConcertApp/MultiUserServer/Source/MultiUserServer/Private/Window/IConcertComponent.h
@@ -17,12 +17,12 @@ struct FConcertComponentInitParams
 	TSharedRef<FConcertServerWindowController> WindowController;
 
 	/** The root area of the main window layout. Add tabs to this. */
-	TSharedRef<FTabManager::FArea> MainWindowArea;
+	TSharedRef<FTabManager::FStack> MainStack;
 
-	FConcertComponentInitParams(TSharedRef<IConcertSyncServer> Server, TSharedRef<FConcertServerWindowController> WindowController, TSharedRef<FTabManager::FArea> MainWindowArea)
+	FConcertComponentInitParams(TSharedRef<IConcertSyncServer> Server, TSharedRef<FConcertServerWindowController> WindowController, TSharedRef<FTabManager::FStack> MainStack)
 		: Server(MoveTemp(Server))
 		, WindowController(MoveTemp(WindowController))
-		, MainWindowArea(MoveTemp(MainWindowArea))
+		, MainStack(MoveTemp(MainStack))
 	{}
 };
 
