Fix CIS issue

diff --git a/Engine/Plugins/Developer/Concert/ConcertApp/MultiUserServer/Source/MultiUserServer/Private/Widgets/SessionTabs/Archived/ArchivedConcertSessionTab.cpp b/Engine/Plugins/Developer/Concert/ConcertApp/MultiUserServer/Source/MultiUserServer/Private/Widgets/SessionTabs/Archived/ArchivedConcertSessionTab.cpp
--- a/Engine/Plugins/Developer/Concert/ConcertApp/MultiUserServer/Source/MultiUserServer/Private/Widgets/SessionTabs/Archived/ArchivedConcertSessionTab.cpp
+++ b/Engine/Plugins/Developer/Concert/ConcertApp/MultiUserServer/Source/MultiUserServer/Private/Widgets/SessionTabs/Archived/ArchivedConcertSessionTab.cpp
@@ -4,27 +4,23 @@
 
 #include "ArchivedSessionHistoryController.h"
 #include "ConcertLogGlobal.h"
-#include "MultiUserServerConsoleVariables.h"
-#include "MultiUserServerModule.h"
-#include "Window/ModalWindowManager.h"
-
-#include "Algo/AllOf.h"
-#include "Algo/Transform.h"
-
-#include "Dialog/SMessageDialog.h"
-
 #include "HistoryEdition/ActivityNode.h"
 #include "HistoryEdition/DebugDependencyGraph.h"
 #include "HistoryEdition/DependencyGraphBuilder.h"
 #include "HistoryEdition/HistoryAnalysis.h"
 #include "HistoryEdition/HistoryDeletion.h"
-
+#include "MultiUserServerConsoleVariables.h"
+#include "MultiUserServerModule.h"
 #include "Session/History/SEditableSessionHistory.h"
 #include "Session/History/SSessionHistory.h"
-
+#include "Window/ModalWindowManager.h"
 #include "Widgets/HistoryDeletion/SDeleteActivityDependenciesDialog.h"
 #include "Widgets/SessionTabs/Archived/SConcertArchivedSessionTabView.h"
-#include "Widgets/StatusBar/SConcertStatusBar.h"
+
+#include "Algo/AllOf.h"
+#include "Algo/Transform.h"
+#include "Dialog/SMessageDialog.h"
+#include "Widgets/Docking/SDockTab.h"
 
 #define LOCTEXT_NAMESPACE "UnrealMultiUserUI"
 
