Create client tab

diff --git a/Engine/Plugins/Developer/Concert/ConcertApp/MultiUserServer/Source/MultiUserServer/Private/MultiUserServerUserSettings.cpp b/Engine/Plugins/Developer/Concert/ConcertApp/MultiUserServer/Source/MultiUserServer/Private/MultiUserServerUserSettings.cpp
deleted file mode 100644
diff --git a/Engine/Plugins/Developer/Concert/ConcertApp/MultiUserServer/Source/MultiUserServer/Private/MultiUserServerUserSettings.h b/Engine/Plugins/Developer/Concert/ConcertApp/MultiUserServer/Source/MultiUserServer/Private/MultiUserServerUserSettings.h
deleted file mode 100644
diff --git a/Engine/Plugins/Developer/Concert/ConcertApp/MultiUserServer/Source/MultiUserServer/Private/Settings/ConcertTransportLogSettings.cpp b/Engine/Plugins/Developer/Concert/ConcertApp/MultiUserServer/Source/MultiUserServer/Private/Settings/ConcertTransportLogSettings.cpp
new file mode 100644
--- /dev/null
+++ b/Engine/Plugins/Developer/Concert/ConcertApp/MultiUserServer/Source/MultiUserServer/Private/Settings/ConcertTransportLogSettings.cpp
@@ -0,0 +1,8 @@
+// Copyright Epic Games, Inc. All Rights Reserved.
+
+#include "ConcertTransportLogSettings.h"
+
+UConcertTransportLogSettings* UConcertTransportLogSettings::GetSettings()
+{
+	return GetMutableDefault<UConcertTransportLogSettings>();
+}
diff --git a/Engine/Plugins/Developer/Concert/ConcertApp/MultiUserServer/Source/MultiUserServer/Private/Settings/ConcertTransportLogSettings.h b/Engine/Plugins/Developer/Concert/ConcertApp/MultiUserServer/Source/MultiUserServer/Private/Settings/ConcertTransportLogSettings.h
new file mode 100644
--- /dev/null
+++ b/Engine/Plugins/Developer/Concert/ConcertApp/MultiUserServer/Source/MultiUserServer/Private/Settings/ConcertTransportLogSettings.h
@@ -0,0 +1,24 @@
+// Copyright Epic Games, Inc. All Rights Reserved.
+
+#pragma once
+
+#include "CoreMinimal.h"
+#include "ConcertFrontendUtils.h"
+#include "UObject/Object.h"
+#include "ConcertTransportLogSettings.generated.h"
+
+/**
+ * 
+ */
+UCLASS(config = MultiUserServerUserSettings, DefaultConfig)
+class MULTIUSERSERVER_API UConcertTransportLogSettings : public UObject
+{
+	GENERATED_BODY()
+public:
+
+	static UConcertTransportLogSettings* GetSettings();
+
+	/** How to displayed FConcertLog::Timestamp */
+	UPROPERTY(Config)
+	ETimeFormat TimestampTimeFormat = ETimeFormat::Absolute;
+};
diff --git a/Engine/Plugins/Developer/Concert/ConcertApp/MultiUserServer/Source/MultiUserServer/Private/Settings/MultiUserServerColumnVisibilitySettings.cpp b/Engine/Plugins/Developer/Concert/ConcertApp/MultiUserServer/Source/MultiUserServer/Private/Settings/MultiUserServerColumnVisibilitySettings.cpp
new file mode 100644
--- /dev/null
+++ b/Engine/Plugins/Developer/Concert/ConcertApp/MultiUserServer/Source/MultiUserServer/Private/Settings/MultiUserServerColumnVisibilitySettings.cpp
@@ -0,0 +1,51 @@
+﻿// Copyright Epic Games, Inc. All Rights Reserved.
+
+#include "MultiUserServerColumnVisibilitySettings.h"
+
+#include "Misc/CoreDelegates.h"
+
+namespace UE::MultiUserServer
+{
+	static bool bIsShutdown = false;
+}
+
+UMultiUserServerColumnVisibilitySettings::UMultiUserServerColumnVisibilitySettings()
+{
+	OnSessionBrowserColumnVisibilityChanged().AddLambda([this](const FColumnVisibilitySnapshot&)
+	{
+		SaveConfig();
+	});
+	OnDeleteActivityDialogColumnVisibility().AddLambda([this](const FColumnVisibilitySnapshot&)
+	{
+		SaveConfig();
+	});
+	OnArchivedActivityBrowserColumnVisibility().AddLambda([this](const FColumnVisibilitySnapshot&)
+	{
+		SaveConfig();
+	});
+	OnLiveActivityBrowserColumnVisibility().AddLambda([this](const FColumnVisibilitySnapshot&)
+	{
+		SaveConfig();
+	});
+	OnLiveSessionContentColumnVisibility().AddLambda([this](const FColumnVisibilitySnapshot&)
+	{
+		SaveConfig();
+	});
+	OnTransportLogColumnVisibility().AddLambda([this](const FColumnVisibilitySnapshot&)
+	{
+		SaveConfig();
+	});
+
+	FCoreDelegates::OnPreExit.AddLambda([]()
+	{
+		UE::MultiUserServer::bIsShutdown = true;
+	});
+}
+
+UMultiUserServerColumnVisibilitySettings* UMultiUserServerColumnVisibilitySettings::GetSettings()
+{
+	// If we're shutting down GetMutableDefault will return garbage - this function may be called by destructors when this module is unloaded
+	return UE::MultiUserServer::bIsShutdown
+		? nullptr
+		: GetMutableDefault<UMultiUserServerColumnVisibilitySettings>();
+}
diff --git a/Engine/Plugins/Developer/Concert/ConcertApp/MultiUserServer/Source/MultiUserServer/Private/Settings/MultiUserServerColumnVisibilitySettings.h b/Engine/Plugins/Developer/Concert/ConcertApp/MultiUserServer/Source/MultiUserServer/Private/Settings/MultiUserServerColumnVisibilitySettings.h
new file mode 100644
--- /dev/null
+++ b/Engine/Plugins/Developer/Concert/ConcertApp/MultiUserServer/Source/MultiUserServer/Private/Settings/MultiUserServerColumnVisibilitySettings.h
@@ -0,0 +1,74 @@
+﻿// Copyright Epic Games, Inc. All Rights Reserved.
+
+#pragma once
+
+#include "CoreMinimal.h"
+#include "UObject/Object.h"
+#include "ConcertHeaderRowUtils.h"
+#include "MultiUserServerColumnVisibilitySettings.generated.h"
+
+/**
+ * 
+ */
+UCLASS(Config = MultiUserServerUserSettings, DefaultConfig)
+class MULTIUSERSERVER_API UMultiUserServerColumnVisibilitySettings : public UObject
+{
+	GENERATED_BODY()
+public:
+
+	DECLARE_MULTICAST_DELEGATE_OneParam(FOnColumnVisibilitySnapshotChanged, const FColumnVisibilitySnapshot& /*NewValue*/);
+
+	UMultiUserServerColumnVisibilitySettings();
+	
+	static UMultiUserServerColumnVisibilitySettings* GetSettings();
+
+	const FColumnVisibilitySnapshot& GetSessionBrowserColumnVisibility() const { return SessionBrowserColumnVisibility; }
+	void SetSessionBrowserColumnVisibility(FColumnVisibilitySnapshot NewValue) { SessionBrowserColumnVisibility = MoveTemp(NewValue); OnSessionBrowserColumnVisibilityChangedEvent.Broadcast(SessionBrowserColumnVisibility); }
+	FOnColumnVisibilitySnapshotChanged& OnSessionBrowserColumnVisibilityChanged() { return OnSessionBrowserColumnVisibilityChangedEvent; }
+
+	const FColumnVisibilitySnapshot& GetDeleteActivityDialogColumnVisibility() const { return DeleteActivityDialogColumnVisibility; }
+	void SetDeleteActivityDialogColumnVisibility(FColumnVisibilitySnapshot NewValue) { DeleteActivityDialogColumnVisibility = MoveTemp(NewValue); OnDeleteActivityDialogColumnVisibilityEvent.Broadcast(DeleteActivityDialogColumnVisibility); }
+	FOnColumnVisibilitySnapshotChanged& OnDeleteActivityDialogColumnVisibility() { return OnDeleteActivityDialogColumnVisibilityEvent; }
+	
+	const FColumnVisibilitySnapshot& GetArchivedActivityBrowserColumnVisibility() const { return ArchivedActivityBrowserColumnVisibility; }
+	void SetArchivedActivityBrowserColumnVisibility(FColumnVisibilitySnapshot NewValue) { ArchivedActivityBrowserColumnVisibility = MoveTemp(NewValue); OnArchivedActivityBrowserColumnVisibilityEvent.Broadcast(ArchivedActivityBrowserColumnVisibility); }
+	FOnColumnVisibilitySnapshotChanged& OnArchivedActivityBrowserColumnVisibility() { return OnArchivedActivityBrowserColumnVisibilityEvent; }
+
+	const FColumnVisibilitySnapshot& GetLiveActivityBrowserColumnVisibility() const { return LiveActivityBrowserColumnVisibility; }
+	void SetLiveActivityBrowserColumnVisibility(FColumnVisibilitySnapshot NewValue) { LiveActivityBrowserColumnVisibility = MoveTemp(NewValue); OnLiveActivityBrowserColumnVisibilityEvent.Broadcast(LiveActivityBrowserColumnVisibility); }
+	FOnColumnVisibilitySnapshotChanged& OnLiveActivityBrowserColumnVisibility() { return OnLiveActivityBrowserColumnVisibilityEvent; }
+
+	const FColumnVisibilitySnapshot& GetLiveSessionContentColumnVisibility() const { return LiveSessionContentColumnVisibility; }
+	void SetLiveSessionContentColumnVisibility(FColumnVisibilitySnapshot NewValue) { LiveSessionContentColumnVisibility = MoveTemp(NewValue); OnLiveSessionContentColumnVisibilityEvent.Broadcast(LiveSessionContentColumnVisibility); }
+	FOnColumnVisibilitySnapshotChanged& OnLiveSessionContentColumnVisibility() { return OnLiveSessionContentColumnVisibilityEvent; }
+
+	const FColumnVisibilitySnapshot& GetTransportLogColumnVisibility() const { return TransportLogColumnVisibility; }
+	void SetTransportLogColumnVisibility(FColumnVisibilitySnapshot NewValue) { TransportLogColumnVisibility = MoveTemp(NewValue); OnTransportLogColumnVisibilityEvent.Broadcast(TransportLogColumnVisibility); }
+	FOnColumnVisibilitySnapshotChanged& OnTransportLogColumnVisibility() { return OnTransportLogColumnVisibilityEvent; }
+	
+private:
+	
+	UPROPERTY(Config)
+	FColumnVisibilitySnapshot SessionBrowserColumnVisibility;
+	FOnColumnVisibilitySnapshotChanged OnSessionBrowserColumnVisibilityChangedEvent;
+
+	UPROPERTY(Config)
+	FColumnVisibilitySnapshot DeleteActivityDialogColumnVisibility;
+	FOnColumnVisibilitySnapshotChanged OnDeleteActivityDialogColumnVisibilityEvent;
+
+	UPROPERTY(Config)
+	FColumnVisibilitySnapshot ArchivedActivityBrowserColumnVisibility;
+	FOnColumnVisibilitySnapshotChanged OnArchivedActivityBrowserColumnVisibilityEvent;
+
+	UPROPERTY(Config)
+	FColumnVisibilitySnapshot LiveActivityBrowserColumnVisibility;
+	FOnColumnVisibilitySnapshotChanged OnLiveActivityBrowserColumnVisibilityEvent;
+
+	UPROPERTY(Config)
+	FColumnVisibilitySnapshot LiveSessionContentColumnVisibility;
+	FOnColumnVisibilitySnapshotChanged OnLiveSessionContentColumnVisibilityEvent;
+
+	UPROPERTY(Config)
+	FColumnVisibilitySnapshot TransportLogColumnVisibility;
+	FOnColumnVisibilitySnapshotChanged OnTransportLogColumnVisibilityEvent;
+};
diff --git a/Engine/Plugins/Developer/Concert/ConcertApp/MultiUserServer/Source/MultiUserServer/Private/Widgets/Browser/ConcertServerSessionBrowserController.cpp b/Engine/Plugins/Developer/Concert/ConcertApp/MultiUserServer/Source/MultiUserServer/Private/Widgets/Browser/ConcertServerSessionBrowserController.cpp
--- a/Engine/Plugins/Developer/Concert/ConcertApp/MultiUserServer/Source/MultiUserServer/Private/Widgets/Browser/ConcertServerSessionBrowserController.cpp
+++ b/Engine/Plugins/Developer/Concert/ConcertApp/MultiUserServer/Source/MultiUserServer/Private/Widgets/Browser/ConcertServerSessionBrowserController.cpp
@@ -249,6 +249,7 @@ TSharedRef<SDockTab> FConcertServerSessionBrowserController::SpawnSessionBrowser
 	const TSharedRef<SDockTab> DockTab = SNew(SDockTab)
 		.Label(LOCTEXT("SessionBrowserTabTitle", "Sessions"))
 		.TabRole(MajorTab)
+		.OnCanCloseTab(SDockTab::FCanCloseTab::CreateLambda([](){ return false; }))
 		[
 			SAssignNew(ConcertBrowser, SConcertServerSessionBrowser, SharedThis(this))
 				.DoubleClickLiveSession(this, &FConcertServerSessionBrowserController::OpenSession)
diff --git a/Engine/Plugins/Developer/Concert/ConcertApp/MultiUserServer/Source/MultiUserServer/Private/Widgets/Browser/SConcertServerSessionBrowser.cpp b/Engine/Plugins/Developer/Concert/ConcertApp/MultiUserServer/Source/MultiUserServer/Private/Widgets/Browser/SConcertServerSessionBrowser.cpp
--- a/Engine/Plugins/Developer/Concert/ConcertApp/MultiUserServer/Source/MultiUserServer/Private/Widgets/Browser/SConcertServerSessionBrowser.cpp
+++ b/Engine/Plugins/Developer/Concert/ConcertApp/MultiUserServer/Source/MultiUserServer/Private/Widgets/Browser/SConcertServerSessionBrowser.cpp
@@ -3,10 +3,10 @@
 #include "SConcertServerSessionBrowser.h"
 
 #include "MultiUserServerModule.h"
-#include "MultiUserServerUserSettings.h"
 #include "Session/Browser/ConcertBrowserUtils.h"
 #include "Session/Browser/ConcertSessionItem.h"
 #include "Session/Browser/SConcertSessionBrowser.h"
+#include "Settings/MultiUserServerColumnVisibilitySettings.h"
 #include "Widgets/Browser/ConcertServerSessionBrowserController.h"
 #include "Window/ConcertServerTabs.h"
 #include "Window/ModalWindowManager.h"
@@ -37,10 +37,10 @@ TSharedRef<SWidget> SConcertServerSessionBrowser::MakeSessionTableView(const FAr
 		.PostRequestedDeleteSession(this, &SConcertServerSessionBrowser::RequestDeleteSession)
 		// Pretend a modal dialog said no - RequestDeleteSession will show non-modal dialog
 		.AskUserToDeleteSessions_Lambda([](auto) { return false; })
-		.ColumnVisibilitySnapshot(UMultiUserServerUserSettings::GetUserSettings()->GetSessionBrowserColumnVisibility())
+		.ColumnVisibilitySnapshot(UMultiUserServerColumnVisibilitySettings::GetSettings()->GetSessionBrowserColumnVisibility())
 		.SaveColumnVisibilitySnapshot_Lambda([](const FColumnVisibilitySnapshot& Snapshot)
 		{
-			UMultiUserServerUserSettings::GetUserSettings()->SetSessionBrowserColumnVisibility(Snapshot);
+			UMultiUserServerColumnVisibilitySettings::GetSettings()->SetSessionBrowserColumnVisibility(Snapshot);
 		});
 }
 
diff --git a/Engine/Plugins/Developer/Concert/ConcertApp/MultiUserServer/Source/MultiUserServer/Private/Widgets/Clients/ConcertClientsTabController.cpp b/Engine/Plugins/Developer/Concert/ConcertApp/MultiUserServer/Source/MultiUserServer/Private/Widgets/Clients/ConcertClientsTabController.cpp
new file mode 100644
--- /dev/null
+++ b/Engine/Plugins/Developer/Concert/ConcertApp/MultiUserServer/Source/MultiUserServer/Private/Widgets/Clients/ConcertClientsTabController.cpp
@@ -0,0 +1,37 @@
+﻿// Copyright Epic Games, Inc. All Rights Reserved.
+
+#include "ConcertClientsTabController.h"
+
+#include "SConcertClientsTabView.h"
+#include "Window/ConcertServerTabs.h"
+#include "Window/ConcertServerWindowController.h"
+
+#define LOCTEXT_NAMESPACE "UnrealMultiUserUI"
+
+void FConcertClientsTabController::Init(const FConcertComponentInitParams& Params)
+{
+	FGlobalTabmanager::Get()->RegisterTabSpawner(
+				ConcertServerTabs::GetClientsTabID(),
+				FOnSpawnTab::CreateRaw(this, &FConcertClientsTabController::SpawnClientsTab, Params.WindowController->GetRootWindow(), Params.Server)
+			)
+			.SetDisplayName(LOCTEXT("ClientsTabTitle", "Clients"))
+			.SetTooltipText(LOCTEXT("ClientsTooltipText", "View network statistics for connected clients.")
+		);
+	Params.MainStack->AddTab(ConcertServerTabs::GetClientsTabID(), ETabState::OpenedTab);
+}
+
+TSharedRef<SDockTab> FConcertClientsTabController::SpawnClientsTab(const FSpawnTabArgs& Args, TSharedPtr<SWindow> RootWindow, TSharedRef<IConcertSyncServer> Server)
+{
+	const TSharedRef<SDockTab> DockTab = SNew(SDockTab)
+		.Label(LOCTEXT("ClientsTabTitle", "Clients"))
+		.TabRole(MajorTab)
+		.CanEverClose(false);
+	DockTab->SetContent(
+		SNew(SConcertClientsTabView, ConcertServerTabs::GetClientsTabID(), Server)
+			.ConstructUnderMajorTab(DockTab)
+			.ConstructUnderWindow(RootWindow)
+		);
+	return DockTab;
+}
+
+#undef LOCTEXT_NAMESPACE
\ No newline at end of file
diff --git a/Engine/Plugins/Developer/Concert/ConcertApp/MultiUserServer/Source/MultiUserServer/Private/Widgets/Clients/ConcertClientsTabController.h b/Engine/Plugins/Developer/Concert/ConcertApp/MultiUserServer/Source/MultiUserServer/Private/Widgets/Clients/ConcertClientsTabController.h
new file mode 100644
--- /dev/null
+++ b/Engine/Plugins/Developer/Concert/ConcertApp/MultiUserServer/Source/MultiUserServer/Private/Widgets/Clients/ConcertClientsTabController.h
@@ -0,0 +1,21 @@
+﻿// Copyright Epic Games, Inc. All Rights Reserved.
+
+#pragma once
+
+#include "CoreMinimal.h"
+#include "Window/IConcertComponent.h"
+
+class SWindow;
+
+class FConcertClientsTabController : public IConcertComponent
+{
+public:
+
+	//~ Begin IConcertComponent Interface
+	virtual void Init(const FConcertComponentInitParams& Params) override;
+	//~ End IConcertComponent Interface
+
+private:
+
+	TSharedRef<SDockTab> SpawnClientsTab(const FSpawnTabArgs& Args, TSharedPtr<SWindow> RootWindow, TSharedRef<IConcertSyncServer>);
+};
diff --git a/Engine/Plugins/Developer/Concert/ConcertApp/MultiUserServer/Source/MultiUserServer/Private/Widgets/Clients/Logging/ConcertLogEntry.h b/Engine/Plugins/Developer/Concert/ConcertApp/MultiUserServer/Source/MultiUserServer/Private/Widgets/Clients/Logging/ConcertLogEntry.h
new file mode 100644
--- /dev/null
+++ b/Engine/Plugins/Developer/Concert/ConcertApp/MultiUserServer/Source/MultiUserServer/Private/Widgets/Clients/Logging/ConcertLogEntry.h
@@ -0,0 +1,17 @@
+// Copyright Epic Games, Inc. All Rights Reserved.
+
+#pragma once
+
+#include "CoreMinimal.h"
+#include "ConcertTransportEvents.h"
+
+using FConcertLogID = uint64;
+
+struct FConcertLogEntry
+{
+	/** Unique log ID. Log IDs grow sequentially. */
+	FConcertLogID LogId;
+
+	/** The log we're describing */
+	FConcertLog Log;
+};
diff --git a/Engine/Plugins/Developer/Concert/ConcertApp/MultiUserServer/Source/MultiUserServer/Private/Widgets/Clients/Logging/Filter/ConcertFrontendLogFilter.h b/Engine/Plugins/Developer/Concert/ConcertApp/MultiUserServer/Source/MultiUserServer/Private/Widgets/Clients/Logging/Filter/ConcertFrontendLogFilter.h
new file mode 100644
--- /dev/null
+++ b/Engine/Plugins/Developer/Concert/ConcertApp/MultiUserServer/Source/MultiUserServer/Private/Widgets/Clients/Logging/Filter/ConcertFrontendLogFilter.h
@@ -0,0 +1,46 @@
+// Copyright Epic Games, Inc. All Rights Reserved.
+
+#pragma once
+
+#include "CoreMinimal.h"
+#include "ConcertLogFilter.h"
+
+/**
+ * A filter that is intended to be displayed in the UI. Every filter has one widget displaying it. */
+class FConcertFrontendLogFilter : public FConcertLogFilter
+{
+public:
+
+	/** Gets a widget that represents this filter */
+	virtual TSharedRef<SWidget> GetFilterWidget() = 0;
+};
+
+/** Helper class to implement filters */
+template<typename TConcertLogFilter, typename TWidgetType = SWidget>
+class TConcertFrontendLogFilterAggregate : public FConcertFrontendLogFilter
+{
+public:
+
+	//~ Begin FConcertLogFilter Interface
+	virtual bool PassesFilter(const FConcertLog& InItem) const final override { return Implementation.PassesFilter(InItem); }
+	//~ End FConcertLogFilter Interface
+	
+	//~ Begin FConcertFrontendLogFilter Interface
+	virtual TSharedRef<SWidget> GetFilterWidget() final override { return ChildSlot.ToSharedRef(); }
+	//~ End FConcertLogFilter Interface
+	
+protected:
+
+	template<typename... TArg>
+	TConcertFrontendLogFilterAggregate(TArg&&... Arg)
+		: Implementation(Forward<TArg>(Arg)...)
+	{
+		Implementation.OnChanged().AddLambda([this]()
+		{
+			OnChanged().Broadcast();
+		});
+	}
+
+	TConcertLogFilter Implementation;
+	TSharedPtr<TWidgetType> ChildSlot;
+};
diff --git a/Engine/Plugins/Developer/Concert/ConcertApp/MultiUserServer/Source/MultiUserServer/Private/Widgets/Clients/Logging/Filter/ConcertFrontendLogFilter_MessageAction.cpp b/Engine/Plugins/Developer/Concert/ConcertApp/MultiUserServer/Source/MultiUserServer/Private/Widgets/Clients/Logging/Filter/ConcertFrontendLogFilter_MessageAction.cpp
new file mode 100644
--- /dev/null
+++ b/Engine/Plugins/Developer/Concert/ConcertApp/MultiUserServer/Source/MultiUserServer/Private/Widgets/Clients/Logging/Filter/ConcertFrontendLogFilter_MessageAction.cpp
@@ -0,0 +1,181 @@
+// Copyright Epic Games, Inc. All Rights Reserved.
+
+#include "ConcertFrontendLogFilter_MessageAction.h"
+
+#include "ConcertTransportEvents.h"
+#include "Widgets/Clients/Logging/Util/MessageActionUtils.h"
+
+#define LOCTEXT_NAMESPACE "UnrealMultiUserUI"
+
+FConcertLogFilter_MessageAction::FConcertLogFilter_MessageAction()
+	: AllowedMessageActionNames(UE::MultiUserServer::MessageActionUtils::GetAllMessageActionNames())
+{}
+
+bool FConcertLogFilter_MessageAction::PassesFilter(const FConcertLog& InItem) const
+{
+	return IsMessageActionAllowed(UE::MultiUserServer::MessageActionUtils::ConvertActionToName(InItem.MessageAction));
+}
+
+void FConcertLogFilter_MessageAction::AllowAll()
+{
+	TSet<FName> Allowed = UE::MultiUserServer::MessageActionUtils::GetAllMessageActionNames();
+	if (Allowed.Num() != AllowedMessageActionNames.Num())
+	{
+		AllowedMessageActionNames = MoveTemp(Allowed);
+		OnChanged().Broadcast();
+	}
+}
+
+void FConcertLogFilter_MessageAction::DisallowAll()
+{
+	if (AllowedMessageActionNames.Num() > 0)
+	{
+		AllowedMessageActionNames.Reset();
+		OnChanged().Broadcast();
+	}
+}
+
+void FConcertLogFilter_MessageAction::ToggleAll(const TSet<FName>& ToToggle)
+{
+	for (const FName MessageTypeName : ToToggle)
+	{
+		if (IsMessageActionAllowed(MessageTypeName))
+		{
+			DisallowMessageAction(MessageTypeName);
+		}
+		else
+		{
+			AllowMessageAction(MessageTypeName);
+		}
+	}
+
+	if (ToToggle.Num() > 0)
+	{
+		OnChanged().Broadcast();
+	}
+}
+
+void FConcertLogFilter_MessageAction::AllowMessageAction(FName MessageTypeName)
+{
+	if (!AllowedMessageActionNames.Contains(MessageTypeName))
+	{
+		AllowedMessageActionNames.Add(MessageTypeName);
+		OnChanged().Broadcast();
+	}
+}
+
+void FConcertLogFilter_MessageAction::DisallowMessageAction(FName MessageTypeName)
+{
+	if (AllowedMessageActionNames.Contains(MessageTypeName))
+	{
+		AllowedMessageActionNames.Remove(MessageTypeName);
+		OnChanged().Broadcast();
+	}
+}
+
+bool FConcertLogFilter_MessageAction::IsMessageActionAllowed(FName MessageTypeName) const
+{
+	return AllowedMessageActionNames.Contains(MessageTypeName);
+}
+
+bool FConcertLogFilter_MessageAction::AreAllAllowed() const
+{
+	return AllowedMessageActionNames.Num() == UE::MultiUserServer::MessageActionUtils::GetAllMessageActionNames().Num();
+}
+
+uint8 FConcertLogFilter_MessageAction::GetNumSelected() const
+{
+	return AllowedMessageActionNames.Num();
+}
+
+FConcertFrontendLogFilter_MessageAction::FConcertFrontendLogFilter_MessageAction()
+{
+	ChildSlot = SNew(SHorizontalBox)
+		.ToolTipText(LOCTEXT("MessageActionFilter.ToolTipText", "Select a list of allowed message actions\nHint: Type in menu to search"))
+
+		+SHorizontalBox::Slot()
+		.AutoWidth()
+		.VAlign(VAlign_Center)
+		[
+			SNew(STextBlock)
+			.Text(LOCTEXT("MessageActionFilter.AllowBefore", "Actions"))
+		]
+
+		+SHorizontalBox::Slot()
+		.AutoWidth()
+		.Padding(2, 0, 0, 0)
+		.VAlign(VAlign_Center)
+		[
+			SNew(SComboButton)
+			.OnGetMenuContent_Raw(this, &FConcertFrontendLogFilter_MessageAction::MakeSelectionMenu)
+			.ButtonContent()
+			[
+				SNew(STextBlock)
+				.Text_Lambda([this]()
+				{
+					return Implementation.AreAllAllowed()
+						? LOCTEXT("MessageActionFilter.Selection.All", "All")
+						: FText::FromString(FString::FromInt(Implementation.GetNumSelected()));
+				})
+			]
+		];
+		
+}
+
+TSharedRef<SWidget> FConcertFrontendLogFilter_MessageAction::MakeSelectionMenu()
+{
+	FMenuBuilder MenuBuilder(false, nullptr);
+
+	MenuBuilder.AddMenuEntry(
+		LOCTEXT("MessageActionFilter.SelectAll.", "Select all"),
+		LOCTEXT("MessageActionFilter.SelectAll.Tooltip", "Allows all message actions"),
+		FSlateIcon(),
+		FUIAction(
+			FExecuteAction::CreateLambda([this](){ Implementation.AllowAll(); }),
+			FCanExecuteAction::CreateLambda([] { return true; }),
+			FIsActionChecked()),
+		NAME_None,
+		EUserInterfaceActionType::Button
+		);
+	MenuBuilder.AddMenuEntry(
+		LOCTEXT("MessageActionFilter.DeselectAll.", "Deselect all"),
+		LOCTEXT("MessageActionFilter.DeelectAll.Tooltip", "Disallows all message actions"),
+		FSlateIcon(),
+		FUIAction(
+			FExecuteAction::CreateLambda([this](){ Implementation.DisallowAll(); }),
+			FCanExecuteAction::CreateLambda([] { return true; }),
+			FIsActionChecked()),
+		NAME_None,
+		EUserInterfaceActionType::Button
+		);
+	
+	MenuBuilder.AddSeparator();
+	
+	for (const FName MessageAction : UE::MultiUserServer::MessageActionUtils::GetAllMessageActionNames())
+	{
+		MenuBuilder.AddMenuEntry(
+			FText::FromString(UE::MultiUserServer::MessageActionUtils::GetActionDisplayString(MessageAction)),
+			FText::GetEmpty(),
+			FSlateIcon(),
+			FUIAction(
+				FExecuteAction::CreateLambda([this, MessageAction]()
+				{
+					if (Implementation.IsMessageActionAllowed(MessageAction))
+					{
+						Implementation.DisallowMessageAction(MessageAction);
+					}
+					else
+					{
+						Implementation.AllowMessageAction(MessageAction);
+					}
+				}),
+				FCanExecuteAction::CreateLambda([] { return true; }),
+				FIsActionChecked::CreateLambda([this, MessageAction]() { return Implementation.IsMessageActionAllowed(MessageAction); })),
+			NAME_None,
+			EUserInterfaceActionType::ToggleButton
+		);
+	}
+	return MenuBuilder.MakeWidget();
+}
+
+#undef LOCTEXT_NAMESPACE
diff --git a/Engine/Plugins/Developer/Concert/ConcertApp/MultiUserServer/Source/MultiUserServer/Private/Widgets/Clients/Logging/Filter/ConcertFrontendLogFilter_MessageAction.h b/Engine/Plugins/Developer/Concert/ConcertApp/MultiUserServer/Source/MultiUserServer/Private/Widgets/Clients/Logging/Filter/ConcertFrontendLogFilter_MessageAction.h
new file mode 100644
--- /dev/null
+++ b/Engine/Plugins/Developer/Concert/ConcertApp/MultiUserServer/Source/MultiUserServer/Private/Widgets/Clients/Logging/Filter/ConcertFrontendLogFilter_MessageAction.h
@@ -0,0 +1,44 @@
+// Copyright Epic Games, Inc. All Rights Reserved.
+
+#pragma once
+
+#include "CoreMinimal.h"
+#include "Widgets/Clients/Logging/Filter/ConcertFrontendLogFilter.h"
+
+/** Allows only the selected messages types */
+class FConcertLogFilter_MessageAction : public FConcertLogFilter
+{
+public:
+
+	FConcertLogFilter_MessageAction();
+
+	//~ Begin FConcertLogFilter Interface
+	virtual bool PassesFilter(const FConcertLog& InItem) const override;
+	//~ End FConcertLogFilter Interface
+
+	void AllowAll();
+	void DisallowAll();
+	void ToggleAll(const TSet<FName>& ToToggle);
+	
+	void AllowMessageAction(FName MessageTypeName);
+	void DisallowMessageAction(FName MessageTypeName);
+	
+	bool IsMessageActionAllowed(FName MessageTypeName) const;
+	bool AreAllAllowed() const;
+	uint8 GetNumSelected() const;
+	
+private:
+	
+	TSet<FName> AllowedMessageActionNames;
+};
+
+class FConcertFrontendLogFilter_MessageAction : public TConcertFrontendLogFilterAggregate<FConcertLogFilter_MessageAction>
+{
+public:
+
+	FConcertFrontendLogFilter_MessageAction();
+
+private:
+
+	TSharedRef<SWidget> MakeSelectionMenu();
+};
diff --git a/Engine/Plugins/Developer/Concert/ConcertApp/MultiUserServer/Source/MultiUserServer/Private/Widgets/Clients/Logging/Filter/ConcertFrontendLogFilter_MessageType.cpp b/Engine/Plugins/Developer/Concert/ConcertApp/MultiUserServer/Source/MultiUserServer/Private/Widgets/Clients/Logging/Filter/ConcertFrontendLogFilter_MessageType.cpp
new file mode 100644
--- /dev/null
+++ b/Engine/Plugins/Developer/Concert/ConcertApp/MultiUserServer/Source/MultiUserServer/Private/Widgets/Clients/Logging/Filter/ConcertFrontendLogFilter_MessageType.cpp
@@ -0,0 +1,228 @@
+// Copyright Epic Games, Inc. All Rights Reserved.
+
+#include "ConcertFrontendLogFilter_MessageType.h"
+
+#include "ConcertTransportEvents.h"
+#include "Widgets/Clients/Logging/Util/MessageTypeUtils.h"
+
+#define LOCTEXT_NAMESPACE "UnrealMultiUserUI"
+
+FConcertLogFilter_MessageType::FConcertLogFilter_MessageType()
+	: AllowedMessageTypeNames(UE::MultiUserServer::MessageTypeUtils::GetAllMessageTypeNames())
+{}
+
+bool FConcertLogFilter_MessageType::PassesFilter(const FConcertLog& InItem) const
+{
+	return IsMessageTypeAllowed(InItem.MessageTypeName);
+}
+
+void FConcertLogFilter_MessageType::AllowAll()
+{
+	TSet<FName> Allowed = UE::MultiUserServer::MessageTypeUtils::GetAllMessageTypeNames();
+	if (Allowed.Num() != AllowedMessageTypeNames.Num())
+	{
+		AllowedMessageTypeNames = MoveTemp(Allowed);
+		OnChanged().Broadcast();
+	}
+}
+
+void FConcertLogFilter_MessageType::DisallowAll()
+{
+	if (AllowedMessageTypeNames.Num() > 0)
+	{
+		AllowedMessageTypeNames.Reset();
+		OnChanged().Broadcast();
+	}
+}
+
+void FConcertLogFilter_MessageType::ToggleAll(const TSet<FName>& ToToggle)
+{
+	for (const FName MessageTypeName : ToToggle)
+	{
+		if (IsMessageTypeAllowed(MessageTypeName))
+		{
+			DisallowMessageType(MessageTypeName);
+		}
+		else
+		{
+			AllowMessageType(MessageTypeName);
+		}
+	}
+
+	if (ToToggle.Num() > 0)
+	{
+		OnChanged().Broadcast();
+	}
+}
+
+void FConcertLogFilter_MessageType::AllowMessageType(FName MessageTypeName)
+{
+	if (!AllowedMessageTypeNames.Contains(MessageTypeName))
+	{
+		AllowedMessageTypeNames.Add(MessageTypeName);
+		OnChanged().Broadcast();
+	}
+}
+
+void FConcertLogFilter_MessageType::DisallowMessageType(FName MessageTypeName)
+{
+	if (AllowedMessageTypeNames.Contains(MessageTypeName))
+	{
+		AllowedMessageTypeNames.Remove(MessageTypeName);
+		OnChanged().Broadcast();
+	}
+}
+
+bool FConcertLogFilter_MessageType::IsMessageTypeAllowed(FName MessageTypeName) const
+{
+	return AllowedMessageTypeNames.Contains(MessageTypeName);
+}
+
+bool FConcertLogFilter_MessageType::AreAllAllowed() const
+{
+	return AllowedMessageTypeNames.Num() == UE::MultiUserServer::MessageTypeUtils::GetAllMessageTypeNames().Num();
+}
+
+uint8 FConcertLogFilter_MessageType::GetNumSelected() const
+{
+	return AllowedMessageTypeNames.Num();
+}
+
+FConcertFrontendLogFilter_MessageType::FConcertFrontendLogFilter_MessageType()
+{
+	ChildSlot = SNew(SHorizontalBox)
+		.ToolTipText(LOCTEXT("MessageTypeFilter.ToolTipText", "Select a list of allowed message types\nHint: Type in menu to search"))
+
+		+SHorizontalBox::Slot()
+		.AutoWidth()
+		.VAlign(VAlign_Center)
+		[
+			SNew(STextBlock)
+			.Text(LOCTEXT("MessageTypeFilter.AllowBefore", "Message Types"))
+		]
+
+		+SHorizontalBox::Slot()
+		.AutoWidth()
+		.Padding(2, 0, 0, 0)
+		.VAlign(VAlign_Center)
+		[
+			SNew(SComboButton)
+			.OnGetMenuContent_Raw(this, &FConcertFrontendLogFilter_MessageType::MakeSelectionMenu)
+			.ButtonContent()
+			[
+				SNew(STextBlock)
+				.Text_Lambda([this]()
+				{
+					return Implementation.AreAllAllowed()
+						? LOCTEXT("MessageTypeFilter.Selection.All", "All")
+						: FText::FromString(FString::FromInt(Implementation.GetNumSelected()));
+				})
+			]
+		];
+		
+}
+
+TSharedRef<SWidget> FConcertFrontendLogFilter_MessageType::MakeSelectionMenu()
+{
+	FMenuBuilder MenuBuilder(false, nullptr);
+
+	MenuBuilder.AddMenuEntry(
+		LOCTEXT("MessageTypeFilter.SelectAll.", "Select all"),
+		LOCTEXT("MessageTypeFilter.SelectAll.Tooltip", "Allows all message types"),
+		FSlateIcon(),
+		FUIAction(
+			FExecuteAction::CreateLambda([this](){ Implementation.AllowAll(); }),
+			FCanExecuteAction::CreateLambda([] { return true; }),
+			FIsActionChecked()),
+		NAME_None,
+		EUserInterfaceActionType::Button
+		);
+	MenuBuilder.AddMenuEntry(
+		LOCTEXT("MessageTypeFilter.DeselectAll.", "Deselect all"),
+		LOCTEXT("MessageTypeFilter.DeelectAll.Tooltip", "Disallows all message types"),
+		FSlateIcon(),
+		FUIAction(
+			FExecuteAction::CreateLambda([this](){ Implementation.DisallowAll(); }),
+			FCanExecuteAction::CreateLambda([] { return true; }),
+			FIsActionChecked()),
+		NAME_None,
+		EUserInterfaceActionType::Button
+		);
+	
+	MenuBuilder.AddSeparator();
+	
+	MenuBuilder.AddMenuEntry(
+		LOCTEXT("MessageTypeFilter.Events.", "Toggle Events"),
+		LOCTEXT("MessageTypeFilter.Events.Tooltip", "Toggles all event data"),
+		FSlateIcon(),
+		FUIAction(
+			FExecuteAction::CreateLambda([this](){ Implementation.ToggleAll(UE::MultiUserServer::MessageTypeUtils::GetAllMessageTypeNames_EventsOnly()); }),
+			FCanExecuteAction::CreateLambda([] { return true; }),
+			FIsActionChecked()),
+		NAME_None,
+		EUserInterfaceActionType::Button
+		);
+	MenuBuilder.AddMenuEntry(
+		LOCTEXT("MessageTypeFilter.Requests.", "Toggle Requests"),
+		LOCTEXT("MessageTypeFilter.Requests.Tooltip", "Toggles all requests"),
+		FSlateIcon(),
+		FUIAction(
+			FExecuteAction::CreateLambda([this](){ Implementation.ToggleAll(UE::MultiUserServer::MessageTypeUtils::GetAllMessageTypeNames_RequestsOnly()); }),
+			FCanExecuteAction::CreateLambda([] { return true; }),
+			FIsActionChecked()),
+		NAME_None,
+		EUserInterfaceActionType::Button
+		);
+	MenuBuilder.AddMenuEntry(
+		LOCTEXT("MessageTypeFilter.Responses.", "Toggle Responses"),
+		LOCTEXT("MessageTypeFilter.Responses.Tooltip", "Toggles all responses"),
+		FSlateIcon(),
+		FUIAction(
+			FExecuteAction::CreateLambda([this](){ Implementation.ToggleAll(UE::MultiUserServer::MessageTypeUtils::GetAllMessageTypeNames_ResponseOnlyOnly()); }),
+			FCanExecuteAction::CreateLambda([] { return true; }),
+			FIsActionChecked()),
+		NAME_None,
+		EUserInterfaceActionType::Button
+		);
+	MenuBuilder.AddMenuEntry(
+		LOCTEXT("MessageTypeFilter.OnlyAck.", "Toggle ACKs"),
+		LOCTEXT("MessageTypeFilter.OnlyAck.Tooltip", "Toggle ACKs"),
+		FSlateIcon(),
+		FUIAction(
+			FExecuteAction::CreateLambda([this](){ Implementation.ToggleAll(UE::MultiUserServer::MessageTypeUtils::GetAllMessageTypeNames_AcksOnlyOnly()); }),
+			FCanExecuteAction::CreateLambda([] { return true; }),
+			FIsActionChecked()),
+		NAME_None,
+		EUserInterfaceActionType::Button
+		);
+	
+	MenuBuilder.AddSeparator();
+	
+	for (const FName MessageType : UE::MultiUserServer::MessageTypeUtils::GetAllMessageTypeNames())
+	{
+		MenuBuilder.AddMenuEntry(
+			FText::FromString(UE::MultiUserServer::MessageTypeUtils::SanitizeMessageTypeName(MessageType)),
+			FText::GetEmpty(),
+			FSlateIcon(),
+			FUIAction(
+				FExecuteAction::CreateLambda([this, MessageType]()
+				{
+					if (Implementation.IsMessageTypeAllowed(MessageType))
+					{
+						Implementation.DisallowMessageType(MessageType);
+					}
+					else
+					{
+						Implementation.AllowMessageType(MessageType);
+					}
+				}),
+				FCanExecuteAction::CreateLambda([] { return true; }),
+				FIsActionChecked::CreateLambda([this, MessageType]() { return Implementation.IsMessageTypeAllowed(MessageType); })),
+			NAME_None,
+			EUserInterfaceActionType::ToggleButton
+		);
+	}
+	return MenuBuilder.MakeWidget();
+}
+
+#undef LOCTEXT_NAMESPACE
diff --git a/Engine/Plugins/Developer/Concert/ConcertApp/MultiUserServer/Source/MultiUserServer/Private/Widgets/Clients/Logging/Filter/ConcertFrontendLogFilter_MessageType.h b/Engine/Plugins/Developer/Concert/ConcertApp/MultiUserServer/Source/MultiUserServer/Private/Widgets/Clients/Logging/Filter/ConcertFrontendLogFilter_MessageType.h
new file mode 100644
--- /dev/null
+++ b/Engine/Plugins/Developer/Concert/ConcertApp/MultiUserServer/Source/MultiUserServer/Private/Widgets/Clients/Logging/Filter/ConcertFrontendLogFilter_MessageType.h
@@ -0,0 +1,44 @@
+// Copyright Epic Games, Inc. All Rights Reserved.
+
+#pragma once
+
+#include "CoreMinimal.h"
+#include "Widgets/Clients/Logging/Filter/ConcertFrontendLogFilter.h"
+
+/** Allows only the selected messages types */
+class FConcertLogFilter_MessageType : public FConcertLogFilter
+{
+public:
+
+	FConcertLogFilter_MessageType();
+
+	//~ Begin FConcertLogFilter Interface
+	virtual bool PassesFilter(const FConcertLog& InItem) const override;
+	//~ End FConcertLogFilter Interface
+
+	void AllowAll();
+	void DisallowAll();
+	void ToggleAll(const TSet<FName>& ToToggle);
+	
+	void AllowMessageType(FName MessageTypeName);
+	void DisallowMessageType(FName MessageTypeName);
+	
+	bool IsMessageTypeAllowed(FName MessageTypeName) const;
+	bool AreAllAllowed() const;
+	uint8 GetNumSelected() const;
+	
+private:
+	
+	TSet<FName> AllowedMessageTypeNames;
+};
+
+class FConcertFrontendLogFilter_MessageType : public TConcertFrontendLogFilterAggregate<FConcertLogFilter_MessageType>
+{
+public:
+
+	FConcertFrontendLogFilter_MessageType();
+
+private:
+
+	TSharedRef<SWidget> MakeSelectionMenu();
+};
diff --git a/Engine/Plugins/Developer/Concert/ConcertApp/MultiUserServer/Source/MultiUserServer/Private/Widgets/Clients/Logging/Filter/ConcertFrontendLogFilter_TextSearch.cpp b/Engine/Plugins/Developer/Concert/ConcertApp/MultiUserServer/Source/MultiUserServer/Private/Widgets/Clients/Logging/Filter/ConcertFrontendLogFilter_TextSearch.cpp
new file mode 100644
--- /dev/null
+++ b/Engine/Plugins/Developer/Concert/ConcertApp/MultiUserServer/Source/MultiUserServer/Private/Widgets/Clients/Logging/Filter/ConcertFrontendLogFilter_TextSearch.cpp
@@ -0,0 +1,36 @@
+// Copyright Epic Games, Inc. All Rights Reserved.
+
+#include "ConcertFrontendLogFilter_TextSearch.h"
+
+#include "ConcertTransportEvents.h"
+#include "Widgets/Clients/Logging/Util/ConcertLogTokenizer.h"
+
+FConcertLogFilter_TextSearch::FConcertLogFilter_TextSearch(TSharedRef<FConcertLogTokenizer> Tokenizer)
+	: TextFilter(TTextFilter<const FConcertLog&>::FItemToStringArray::CreateRaw(this, &FConcertLogFilter_TextSearch::GenerateSearchTerms))
+	, Tokenizer(MoveTemp(Tokenizer))
+{
+	TextFilter.OnChanged().AddLambda([this]
+	{
+		OnChanged().Broadcast();
+	});
+}
+
+void FConcertLogFilter_TextSearch::GenerateSearchTerms(const FConcertLog& InItem, TArray<FString>& OutTerms) const
+{
+	for (TFieldIterator<const FProperty> PropertyIt(FConcertLog::StaticStruct()); PropertyIt; ++PropertyIt)
+	{
+		OutTerms.Add(Tokenizer->Tokenize(InItem, **PropertyIt));
+	}
+}
+
+FConcertFrontendLogFilter_TextSearch::FConcertFrontendLogFilter_TextSearch(TSharedRef<FConcertLogTokenizer> Tokenizer)
+	: Super(MoveTemp(Tokenizer))
+{
+	ChildSlot = SNew(SSearchBox)
+		.OnTextChanged_Lambda([this](const FText& NewSearchText)
+		{
+			OnSearchTextChanged().Broadcast(NewSearchText);
+			Implementation.SetRawFilterText(NewSearchText);
+		})
+		.DelayChangeNotificationsWhileTyping(true);
+}
diff --git a/Engine/Plugins/Developer/Concert/ConcertApp/MultiUserServer/Source/MultiUserServer/Private/Widgets/Clients/Logging/Filter/ConcertFrontendLogFilter_TextSearch.h b/Engine/Plugins/Developer/Concert/ConcertApp/MultiUserServer/Source/MultiUserServer/Private/Widgets/Clients/Logging/Filter/ConcertFrontendLogFilter_TextSearch.h
new file mode 100644
--- /dev/null
+++ b/Engine/Plugins/Developer/Concert/ConcertApp/MultiUserServer/Source/MultiUserServer/Private/Widgets/Clients/Logging/Filter/ConcertFrontendLogFilter_TextSearch.h
@@ -0,0 +1,51 @@
+// Copyright Epic Games, Inc. All Rights Reserved.
+
+#pragma once
+
+#include "CoreMinimal.h"
+#include "ConcertFrontendLogFilter.h"
+#include "Misc/TextFilter.h"
+#include "Widgets/Input/SSearchBox.h"
+
+class FConcertLogTokenizer;
+
+/** Allows advanced search by text. Implements Adapter pattern to wrap TTextFilter. */
+class FConcertLogFilter_TextSearch : public FConcertLogFilter
+{
+public:
+
+	FConcertLogFilter_TextSearch(TSharedRef<FConcertLogTokenizer> Tokenizer);
+	
+	//~ Begin FConcertLogFilter Interface
+	virtual bool PassesFilter(const FConcertLog& InItem) const override { return TextFilter.PassesFilter(InItem); }
+	//~ End FConcertLogFilter Interface
+	
+	void SetRawFilterText(const FText& InFilterText) { TextFilter.SetRawFilterText(InFilterText); }
+	
+private:
+
+	/** Does the actual string search */
+	TTextFilter<const FConcertLog&> TextFilter;
+	/** Helps in converting FConcertLog members into search terms */
+	TSharedRef<FConcertLogTokenizer> Tokenizer;
+
+	/** Parses InItem into a bunch of strings that can be searched */
+	void GenerateSearchTerms(const FConcertLog& InItem, TArray<FString>& OutTerms) const;
+};
+
+/** Creates a search bar */
+class FConcertFrontendLogFilter_TextSearch : public TConcertFrontendLogFilterAggregate<FConcertLogFilter_TextSearch, SSearchBox>
+{
+	using Super = TConcertFrontendLogFilterAggregate<FConcertLogFilter_TextSearch, SSearchBox>;
+public:
+
+	FConcertFrontendLogFilter_TextSearch(TSharedRef<FConcertLogTokenizer> Tokenizer);
+
+	DECLARE_MULTICAST_DELEGATE_OneParam(FOnSearchTextChanged, const FText& /*NewSearchText*/);
+	FOnSearchTextChanged& OnSearchTextChanged() { return OnSearchTextChangedEvent; }
+
+private:
+
+	/** Useful to let external subscriber handle text highlighting */
+	FOnSearchTextChanged OnSearchTextChangedEvent;
+};
diff --git a/Engine/Plugins/Developer/Concert/ConcertApp/MultiUserServer/Source/MultiUserServer/Private/Widgets/Clients/Logging/Filter/ConcertFrontendLogFilter_Time.cpp b/Engine/Plugins/Developer/Concert/ConcertApp/MultiUserServer/Source/MultiUserServer/Private/Widgets/Clients/Logging/Filter/ConcertFrontendLogFilter_Time.cpp
new file mode 100644
--- /dev/null
+++ b/Engine/Plugins/Developer/Concert/ConcertApp/MultiUserServer/Source/MultiUserServer/Private/Widgets/Clients/Logging/Filter/ConcertFrontendLogFilter_Time.cpp
@@ -0,0 +1,138 @@
+// Copyright Epic Games, Inc. All Rights Reserved.
+
+#include "ConcertFrontendLogFilter_Time.h"
+
+#include "ConcertFrontendUtils.h"
+#include "ConcertTransportEvents.h"
+
+#include "Widgets/Input/SEditableTextBox.h"
+#include "Widgets/SBoxPanel.h"
+#include "Widgets/Text/STextBlock.h"
+
+#define LOCTEXT_NAMESPACE "UnrealMultiUserUI"
+
+FConcertLogFilter_Time::FConcertLogFilter_Time(ETimeFilter FilterMode)
+	: FilterMode(FilterMode)
+	// Make the filter allow everything by default
+	, Time(MakeResetTime())
+{}
+
+void FConcertLogFilter_Time::ResetToInfiniteTime()
+{
+	SetTime(MakeResetTime());
+}
+
+FDateTime FConcertLogFilter_Time::MakeResetTime() const
+{
+	return FilterMode == ETimeFilter::AllowAfter ? FDateTime() : FDateTime::MaxValue();
+}
+
+bool FConcertLogFilter_Time::PassesFilter(const FConcertLog& InItem) const
+{
+	switch (FilterMode)
+	{
+	case ETimeFilter::AllowAfter: return InItem.Timestamp >= Time;
+	case ETimeFilter::AllowBefore: return InItem.Timestamp <= Time;
+	default:
+		checkNoEntry();
+		return true;
+	}
+}
+
+void FConcertLogFilter_Time::SetFilterMode(ETimeFilter InFilterMode)
+{
+	if (FilterMode != InFilterMode)
+	{
+		FilterMode = InFilterMode;
+		OnChanged().Broadcast();
+	}
+}
+
+void FConcertLogFilter_Time::SetTime(const FDateTime& InTime)
+{
+	if (Time != InTime)
+	{
+		Time = InTime;
+		OnChanged().Broadcast();
+	}
+}
+
+FConcertFrontendLogFilter_Time::FConcertFrontendLogFilter_Time(ETimeFilter TimeFilter)
+	: Super(TimeFilter)
+{
+	ChildSlot = SNew(SHorizontalBox)
+		.ToolTipText(LOCTEXT("TimeFilter.ToolTipText", "Filter logs by local time"))
+
+		+SHorizontalBox::Slot()
+		.AutoWidth()
+		.VAlign(VAlign_Center)
+		[
+			SNew(STextBlock)
+			.Text_Lambda([this]()
+			{
+				return Implementation.GetFilterMode() == ETimeFilter::AllowAfter
+					? LOCTEXT("TimeFilter.AllowAfter", "After")
+					: LOCTEXT("TimeFilter.AllowBefore", "Before");
+			})
+		]
+
+		+SHorizontalBox::Slot()
+		.AutoWidth()
+		.Padding(2, 0, 0, 0)
+		.VAlign(VAlign_Center)
+		[
+			SNew(SComboButton)
+			.OnGetMenuContent_Raw(this, &FConcertFrontendLogFilter_Time::CreateDatePicker)
+			.ButtonContent()
+			[
+				SNew(STextBlock)
+				.Text_Lambda([this](){ return FText::AsDateTime(Implementation.GetTime()); })
+			]
+		];
+}
+
+TSharedRef<SWidget> FConcertFrontendLogFilter_Time::CreateDatePicker()
+{
+	FMenuBuilder MenuBuilder(true, nullptr);
+
+	MenuBuilder.AddMenuEntry(
+		LOCTEXT("TimeFilter.Clear.", "Clear"),
+		LOCTEXT("TimeFilter.Clear.Tooltip", "Sets the time so that this filter has no effect"),
+		FSlateIcon(),
+		FUIAction(
+			FExecuteAction::CreateLambda([this](){ Implementation.ResetToInfiniteTime(); }),
+			FCanExecuteAction::CreateLambda([] { return true; }),
+			FIsActionChecked()),
+		NAME_None,
+		EUserInterfaceActionType::Button
+		);
+	MenuBuilder.AddMenuEntry(
+		LOCTEXT("TimeFilter.Now.", "Now"),
+		LOCTEXT("TimeFilter.Now.Tooltip", "Sets the time to now in local time"),
+		FSlateIcon(),
+		FUIAction(
+			FExecuteAction::CreateLambda([this](){ Implementation.SetTime(FDateTime::UtcNow()); }),
+			FCanExecuteAction::CreateLambda([] { return true; }),
+			FIsActionChecked()),
+		NAME_None,
+		EUserInterfaceActionType::Button
+	);
+	MenuBuilder.AddSeparator();
+
+	MenuBuilder.AddWidget(
+		SNew(SEditableTextBox)
+		.Text(FText::FromString(Implementation.GetTime().ToString()))
+		.OnTextCommitted_Lambda([this](const FText& Text, ETextCommit::Type)
+		{
+			FDateTime Time;
+			if (FDateTime::Parse(Text.ToString(), Time))
+			{
+				Implementation.SetTime(Time);
+			}
+		}),
+		LOCTEXT("TimeFilter.Custom.", "Custom time")
+		);
+	return MenuBuilder.MakeWidget();
+}
+
+#undef LOCTEXT_NAMESPACE
diff --git a/Engine/Plugins/Developer/Concert/ConcertApp/MultiUserServer/Source/MultiUserServer/Private/Widgets/Clients/Logging/Filter/ConcertFrontendLogFilter_Time.h b/Engine/Plugins/Developer/Concert/ConcertApp/MultiUserServer/Source/MultiUserServer/Private/Widgets/Clients/Logging/Filter/ConcertFrontendLogFilter_Time.h
new file mode 100644
--- /dev/null
+++ b/Engine/Plugins/Developer/Concert/ConcertApp/MultiUserServer/Source/MultiUserServer/Private/Widgets/Clients/Logging/Filter/ConcertFrontendLogFilter_Time.h
@@ -0,0 +1,53 @@
+// Copyright Epic Games, Inc. All Rights Reserved.
+
+#pragma once
+
+#include "CoreMinimal.h"
+#include "ConcertFrontendLogFilter.h"
+
+enum class ETimeFilter
+{
+	/** Logs after the indicated time are allowed */
+	AllowAfter,
+	/** Logs before the indicated time are allowed */
+	AllowBefore
+};
+
+/** Filters based on whether a log happened before or after a certain time */
+class FConcertLogFilter_Time : public FConcertLogFilter
+{
+public:
+
+	FConcertLogFilter_Time(ETimeFilter FilterMode);
+
+	void ResetToInfiniteTime();
+
+	//~ Begin FConcertLogFilter Interface
+	virtual bool PassesFilter(const FConcertLog& InItem) const override;
+	//~ End FConcertLogFilter Interfac
+
+	ETimeFilter GetFilterMode() const { return FilterMode; }
+	FDateTime GetTime() const { return Time; }
+
+	void SetFilterMode(ETimeFilter InFilterMode);
+	void SetTime(const FDateTime& InTime);
+	
+private:
+
+	ETimeFilter FilterMode;
+	FDateTime Time;
+
+	FDateTime MakeResetTime() const;
+};
+
+class FConcertFrontendLogFilter_Time : public TConcertFrontendLogFilterAggregate<FConcertLogFilter_Time>
+{
+	using Super = TConcertFrontendLogFilterAggregate<FConcertLogFilter_Time>;
+public:
+	
+	FConcertFrontendLogFilter_Time(ETimeFilter TimeFilter);
+
+private:
+
+	TSharedRef<SWidget> CreateDatePicker();
+};
diff --git a/Engine/Plugins/Developer/Concert/ConcertApp/MultiUserServer/Source/MultiUserServer/Private/Widgets/Clients/Logging/Filter/ConcertLogFilter.cpp b/Engine/Plugins/Developer/Concert/ConcertApp/MultiUserServer/Source/MultiUserServer/Private/Widgets/Clients/Logging/Filter/ConcertLogFilter.cpp
new file mode 100644
--- /dev/null
+++ b/Engine/Plugins/Developer/Concert/ConcertApp/MultiUserServer/Source/MultiUserServer/Private/Widgets/Clients/Logging/Filter/ConcertLogFilter.cpp
@@ -0,0 +1,3 @@
+// Copyright Epic Games, Inc. All Rights Reserved.
+
+#include "ConcertLogFilter.h"
diff --git a/Engine/Plugins/Developer/Concert/ConcertApp/MultiUserServer/Source/MultiUserServer/Private/Widgets/Clients/Logging/Filter/ConcertLogFilter.h b/Engine/Plugins/Developer/Concert/ConcertApp/MultiUserServer/Source/MultiUserServer/Private/Widgets/Clients/Logging/Filter/ConcertLogFilter.h
new file mode 100644
--- /dev/null
+++ b/Engine/Plugins/Developer/Concert/ConcertApp/MultiUserServer/Source/MultiUserServer/Private/Widgets/Clients/Logging/Filter/ConcertLogFilter.h
@@ -0,0 +1,31 @@
+// Copyright Epic Games, Inc. All Rights Reserved.
+
+#pragma once
+
+#include "CoreMinimal.h"
+#include "Misc/IFilter.h"
+
+struct FConcertLog;
+
+/** Base filter */
+class FConcertLogFilter :
+	public IFilter<const FConcertLog&>,
+	public TSharedFromThis<FConcertLogFilter>,
+	// We do not need to copy filters so let's avoid it by accident; some constructors pass this pointer to callbacks which become stale upon copying
+	public FNoncopyable
+{
+public:
+
+	//~ Begin IFilter Interface
+	DECLARE_DERIVED_EVENT( FFrontendFilter, IFilter<const FConcertLog&>::FChangedEvent, FChangedEvent );
+	virtual FChangedEvent& OnChanged() override { return ChangedEvent; }
+	//~ End IFilter Interface
+
+protected:
+	
+	void BroadcastChangedEvent() const { ChangedEvent.Broadcast(); }
+
+private:
+	
+	FChangedEvent ChangedEvent;
+};
diff --git a/Engine/Plugins/Developer/Concert/ConcertApp/MultiUserServer/Source/MultiUserServer/Private/Widgets/Clients/Logging/Filter/ConcertLogFilter_FrontendRoot.cpp b/Engine/Plugins/Developer/Concert/ConcertApp/MultiUserServer/Source/MultiUserServer/Private/Widgets/Clients/Logging/Filter/ConcertLogFilter_FrontendRoot.cpp
new file mode 100644
--- /dev/null
+++ b/Engine/Plugins/Developer/Concert/ConcertApp/MultiUserServer/Source/MultiUserServer/Private/Widgets/Clients/Logging/Filter/ConcertLogFilter_FrontendRoot.cpp
@@ -0,0 +1,98 @@
+// Copyright Epic Games, Inc. All Rights Reserved.
+
+#include "ConcertLogFilter_FrontendRoot.h"
+
+#include "ConcertFrontendLogFilter.h"
+#include "ConcertFrontendLogFilter_MessageAction.h"
+#include "ConcertFrontendLogFilter_MessageType.h"
+#include "ConcertFrontendLogFilter_TextSearch.h"
+#include "ConcertFrontendLogFilter_Time.h"
+#include "Algo/AllOf.h"
+
+#include "Algo/AnyOf.h"
+#include "Widgets/SBoxPanel.h"
+
+FConcertLogFilter_FrontendRoot::FConcertLogFilter_FrontendRoot(TSharedRef<FConcertLogTokenizer> Tokenizer, TArray<TSharedRef<FConcertFrontendLogFilter>> InCustomFilters)
+	: TextSearchFilter(MakeShared<FConcertFrontendLogFilter_TextSearch>(MoveTemp(Tokenizer)))
+	, CustomFilters(MoveTemp(InCustomFilters))
+	, AllFrontendFilters(CustomFilters)
+{
+	AllFrontendFilters.Add(TextSearchFilter);
+	
+	for (const TSharedRef<FConcertFrontendLogFilter>& Filter : AllFrontendFilters)
+	{
+		Filter->OnChanged().AddLambda([this]()
+		{
+			OnChanged().Broadcast();
+		});
+	}
+}
+
+TSharedRef<SWidget> FConcertLogFilter_FrontendRoot::BuildFilterWidgets() const
+{
+	return SNew(SVerticalBox)
+
+		// Search bar
+		+SVerticalBox::Slot()
+		.AutoHeight()
+		[
+			TextSearchFilter->GetFilterWidget()
+		]
+
+		+SVerticalBox::Slot()
+		.AutoHeight()
+		.Padding(0, 2)
+		[
+			BuildCustomFilterListWidget()
+		];
+}
+
+bool FConcertLogFilter_FrontendRoot::PassesFilter(const FConcertLog& InItem) const
+{
+	return Algo::AllOf(
+		AllFrontendFilters,
+		[&InItem](const TSharedRef<FConcertFrontendLogFilter>& AndFilter){ return AndFilter->PassesFilter(InItem); }
+		);
+}
+
+TSharedRef<SWidget> FConcertLogFilter_FrontendRoot::BuildCustomFilterListWidget() const
+{
+	TSharedRef<SHorizontalBox> Box = SNew(SHorizontalBox);
+
+	bool bIsFirst = true;
+	for (const TSharedRef<FConcertFrontendLogFilter>& Filter : CustomFilters)
+	{
+		const FMargin Margin = bIsFirst ? FMargin() : FMargin(8, 0, 0, 0);
+		bIsFirst = false;
+		Box->AddSlot()
+			.AutoWidth()
+			.Padding(Margin)
+			.VAlign(VAlign_Center)
+			[
+				Filter->GetFilterWidget()
+			];
+	}
+
+	return Box;
+}
+
+namespace UE::MultiUserServer::Private
+{
+	static TArray<TSharedRef<FConcertFrontendLogFilter>> CreateCommonFilters()
+	{
+		return {
+			MakeShared<FConcertFrontendLogFilter_MessageAction>(),
+			MakeShared<FConcertFrontendLogFilter_MessageType>(),
+			MakeShared<FConcertFrontendLogFilter_Time>(ETimeFilter::AllowAfter),
+			MakeShared<FConcertFrontendLogFilter_Time>(ETimeFilter::AllowBefore)
+		};
+	}
+}
+
+TSharedRef<FConcertLogFilter_FrontendRoot> UE::MultiUserServer::MakeGlobalLogFilter(TSharedRef<FConcertLogTokenizer> Tokenizer)
+{
+	return MakeShared<FConcertLogFilter_FrontendRoot>(
+		MoveTemp(Tokenizer),
+		Private::CreateCommonFilters()
+		);
+}
diff --git a/Engine/Plugins/Developer/Concert/ConcertApp/MultiUserServer/Source/MultiUserServer/Private/Widgets/Clients/Logging/Filter/ConcertLogFilter_FrontendRoot.h b/Engine/Plugins/Developer/Concert/ConcertApp/MultiUserServer/Source/MultiUserServer/Private/Widgets/Clients/Logging/Filter/ConcertLogFilter_FrontendRoot.h
new file mode 100644
--- /dev/null
+++ b/Engine/Plugins/Developer/Concert/ConcertApp/MultiUserServer/Source/MultiUserServer/Private/Widgets/Clients/Logging/Filter/ConcertLogFilter_FrontendRoot.h
@@ -0,0 +1,48 @@
+// Copyright Epic Games, Inc. All Rights Reserved.
+
+#pragma once
+
+#include "CoreMinimal.h"
+#include "ConcertLogFilter.h"
+
+class FConcertLogTokenizer;
+class FConcertFrontendLogFilter;
+class FConcertFrontendLogFilter_TextSearch;
+
+/** A filter that contains multiple UI filters */
+class FConcertLogFilter_FrontendRoot : public FConcertLogFilter, public TSharedFromThis<FConcertLogFilter_FrontendRoot>
+{
+public:
+
+	FConcertLogFilter_FrontendRoot(TSharedRef<FConcertLogTokenizer> Tokenizer, TArray<TSharedRef<FConcertFrontendLogFilter>> InCustomFilters);
+
+	/** Builds the widget view for all contained filters */
+	TSharedRef<SWidget> BuildFilterWidgets() const;
+	
+	//~ Begin IFilter Interface
+	virtual bool PassesFilter(const FConcertLog& InItem) const override;
+	//~ End IFilter Interface
+
+	FORCEINLINE const TSharedRef<FConcertFrontendLogFilter_TextSearch>& GetTextSearchFilter() const { return TextSearchFilter; }
+	FORCEINLINE const TArray<TSharedRef<FConcertFrontendLogFilter>>& GetFrontendFilters() const { return AllFrontendFilters; }
+
+private:
+
+	/** The text search filter. Also in FrontendFilters. Separate variable to build search bar in new line. */
+	TSharedRef<FConcertFrontendLogFilter_TextSearch> TextSearchFilter;
+
+	/** AllFrontendFilters without special filters we have as propteries above, such as TextSearchFilter. */
+	TArray<TSharedRef<FConcertFrontendLogFilter>> CustomFilters;
+	
+	/** Filters that are combined using logical OR. */
+	TArray<TSharedRef<FConcertFrontendLogFilter>> AllFrontendFilters;
+
+	/** Builds the widgets that go under the text */
+	TSharedRef<SWidget> BuildCustomFilterListWidget() const;
+};
+
+namespace UE::MultiUserServer
+{
+	/** Creates a filter for the global filter log window */
+	TSharedRef<FConcertLogFilter_FrontendRoot> MakeGlobalLogFilter(TSharedRef<FConcertLogTokenizer> Tokenizer);
+}
diff --git a/Engine/Plugins/Developer/Concert/ConcertApp/MultiUserServer/Source/MultiUserServer/Private/Widgets/Clients/Logging/Filter/FilteredConcertLogList.cpp b/Engine/Plugins/Developer/Concert/ConcertApp/MultiUserServer/Source/MultiUserServer/Private/Widgets/Clients/Logging/Filter/FilteredConcertLogList.cpp
new file mode 100644
--- /dev/null
+++ b/Engine/Plugins/Developer/Concert/ConcertApp/MultiUserServer/Source/MultiUserServer/Private/Widgets/Clients/Logging/Filter/FilteredConcertLogList.cpp
@@ -0,0 +1,179 @@
+// Copyright Epic Games, Inc. All Rights Reserved.
+
+#include "FilteredConcertLogList.h"
+
+#include "ConcertLogFilter.h"
+
+FFilteredConcertLogList::FFilteredConcertLogList(TSharedRef<IConcertLogSource> InLogSource, TSharedPtr<FConcertLogFilter> InOptionalFilter)
+	: LogSource(MoveTemp(InLogSource))
+	, Filter(MoveTemp(InOptionalFilter))
+{
+	LogSource->OnLowestLogEntryChanged().AddRaw(this, &FFilteredConcertLogList::OnLowestLogEntryChanged);
+	LogSource->OnLogEntryAdded().AddRaw(this, &FFilteredConcertLogList::OnNewLogEntryAdded);
+
+	if (Filter)
+	{
+		Filter->OnChanged().AddRaw(this, &FFilteredConcertLogList::RebuildFilteredResult);
+	}
+}
+
+FFilteredConcertLogList::~FFilteredConcertLogList()
+{
+	// This is not really needed because the LogSource is SUPPOSED to be solely owned by FFilteredConcertLogList.
+	// We do it for safety regardless...
+	LogSource->OnLowestLogEntryChanged().RemoveAll(this);
+	LogSource->OnLogEntryAdded().RemoveAll(this);
+
+	// ... Filter is not uniquely owned but removing effectively also does not really matter.
+	if (Filter)
+	{
+		Filter->OnChanged().RemoveAll(this);
+	}
+}
+
+void FFilteredConcertLogList::OnLowestLogEntryChanged(FConcertLogID NewLowestValidID)
+{
+	for (size_t i = 0; i < FilteredLogs.Num() && FilteredLogs[i]->LogId < NewLowestValidID; ++i)
+	{
+		// Worst case this is O(n^2) - usually a thought O(cn) with small c)...
+		// Not much we without rewriting SListView to use something other than TArray
+		FilteredLogs.RemoveAt(i);
+		--i;
+	}
+}
+
+void FFilteredConcertLogList::OnNewLogEntryAdded(const TSharedRef<FConcertLogEntry>& NewLogEntry)
+{
+	if (!Filter || Filter->PassesFilter(NewLogEntry->Log))
+	{
+		FilteredLogs.Emplace(NewLogEntry);
+		OnLogListChanged().Broadcast(GetFilteredLogs());
+	}
+}
+
+void FFilteredConcertLogList::RebuildFilteredResult()
+{
+	FilteredLogs.Reset();
+	LogSource->ForEachLog([this](const TSharedPtr<FConcertLogEntry>& LogEntry)
+	{
+		if (!Filter || Filter->PassesFilter(LogEntry->Log))
+		{
+			FilteredLogs.Emplace(LogEntry);
+		}
+	});
+
+	OnLogListChanged().Broadcast(GetFilteredLogs());
+}
+
+FPagedFilteredConcertLogList::FPagedFilteredConcertLogList(TSharedRef<IConcertLogSource> LogSource, TSharedPtr<FConcertLogFilter> OptionalFilter, uint16 InitialLogsPerPage)
+	: FFilteredConcertLogList(MoveTemp(LogSource), MoveTemp(OptionalFilter))
+	, LogsPerPage(FMath::Max<uint16>(InitialLogsPerPage, 1))
+{
+	checkf(InitialLogsPerPage >= 1, TEXT("Unreasonable page size"));
+
+	OnLogListChanged().AddRaw(this, &FPagedFilteredConcertLogList::HandeOnLogListChanged);
+}
+
+void FPagedFilteredConcertLogList::SetLogsPerPage(uint16 NewLogsPerPage)
+{
+	if (LogsPerPage != NewLogsPerPage && ensure(NewLogsPerPage > 0))
+	{
+		size_t StartIndex, LastIndex;
+		Tie(StartIndex, LastIndex) = GetLogIndicesForPage();
+		
+		LogsPerPage = NewLogsPerPage;
+		// Recalculate current page so the first item that was the page is still on the page
+		CurrentPageIndex = StartIndex / LogsPerPage;
+
+		RepopulatePage();
+	}
+}
+
+void FPagedFilteredConcertLogList::SetPage(const uint32 PageIndex)
+{
+	if (CurrentPageIndex != PageIndex && PageIndex < GetNumPages())
+	{
+		CurrentPageIndex = PageIndex;
+		RepopulatePage();
+	}
+}
+
+void FPagedFilteredConcertLogList::HandeOnLogListChanged(const TArray<TSharedPtr<FConcertLogEntry>>& NewFilteredLogList)
+{
+	CheckAndConditionallyPopulatePage();
+}
+
+void FPagedFilteredConcertLogList::RepopulatePage()
+{
+	PageView.Empty(LogsPerPage);
+	ForEachLogIndexOnPage([this](size_t i)
+	{
+		PageView.Emplace(GetFilteredLogs()[i]);
+	});
+
+	OnPageViewChanged().Broadcast(PageView);
+}
+
+void FPagedFilteredConcertLogList::CheckAndConditionallyPopulatePage()
+{
+	bool bChanged = false;
+	size_t PageViewIndex = 0;
+	ForEachLogIndexOnPage([this, &PageViewIndex, &bChanged](size_t i)
+	{
+		if (PageView[PageViewIndex] != GetFilteredLogs()[i])
+		{
+			PageView[PageViewIndex] = GetFilteredLogs()[i];
+			bChanged = true;
+		}
+
+		++PageViewIndex;
+	}, FMath::Min(PageView.Num(), GetFilteredLogs().Num()));
+
+	size_t StartIndex, LastIndex;
+	Tie(StartIndex, LastIndex) = GetLogIndicesForPage();
+
+	// Add items that are new in GetFilteredLogs()
+	bChanged |= PageView.Num() <= LastIndex; 
+	for (size_t i = PageView.Num(); i <= LastIndex && GetFilteredLogs().IsValidIndex(i); ++i)
+	{
+		PageView.Emplace(GetFilteredLogs()[i]);
+	}
+
+	// Remove (possibly duplicate) items that exceed the expected page length 
+	const FLogsPerPageCount LogsOnPage = FMath::Min(LastIndex, LastIndex - StartIndex + 1);
+	const size_t NumToRemove = PageView.Num() - LogsOnPage;
+	bChanged |= NumToRemove > 0;
+	for (size_t i = 0; i < NumToRemove; ++i)
+	{
+		PageView.RemoveAt(PageView.Num() - 1);
+	}
+
+	if (bChanged)
+	{
+		OnPageViewChanged().Broadcast(PageView);
+	}
+}
+
+void FPagedFilteredConcertLogList::ForEachLogIndexOnPage(TFunctionRef<void(size_t Index)> Callback, FLogsPerPageCount MaxItems)
+{
+	size_t StartIndex, LastIndex;
+	Tie(StartIndex, LastIndex) = GetLogIndicesForPage();
+	const FLogsPerPageCount NumItems = LastIndex - StartIndex + 1;
+	for (size_t i = 0; i < NumItems && i < MaxItems; ++i)
+	{
+		Callback(StartIndex + i);
+	}
+}
+
+TTuple<size_t, size_t> FPagedFilteredConcertLogList::GetLogIndicesForPage() const
+{
+	const size_t StartIndex = CurrentPageIndex * LogsPerPage;
+	const size_t LastIndex = FMath::Min<size_t>(
+		// Used when page is full
+		StartIndex + LogsPerPage - 1,
+		// Used when page is not full
+		GetFilteredLogs().Num() == 0 ? 0 : GetFilteredLogs().Num() - 1
+		);
+	check(GetFilteredLogs().Num() == 0 || GetFilteredLogs().IsValidIndex(StartIndex));
+	return { StartIndex, LastIndex };
+}
diff --git a/Engine/Plugins/Developer/Concert/ConcertApp/MultiUserServer/Source/MultiUserServer/Private/Widgets/Clients/Logging/Filter/FilteredConcertLogList.h b/Engine/Plugins/Developer/Concert/ConcertApp/MultiUserServer/Source/MultiUserServer/Private/Widgets/Clients/Logging/Filter/FilteredConcertLogList.h
new file mode 100644
--- /dev/null
+++ b/Engine/Plugins/Developer/Concert/ConcertApp/MultiUserServer/Source/MultiUserServer/Private/Widgets/Clients/Logging/Filter/FilteredConcertLogList.h
@@ -0,0 +1,94 @@
+// Copyright Epic Games, Inc. All Rights Reserved.
+
+#pragma once
+
+#include "CoreMinimal.h"
+#include "Widgets/Clients/Logging/Source/IConcertLogSource.h"
+
+class FConcertLogFilter;
+struct FConcertLog;
+
+/** Applies a filter to a IConcertLogSource as the logs come in. */
+class FFilteredConcertLogList
+{
+public:
+	
+	FFilteredConcertLogList(TSharedRef<IConcertLogSource> InLogSource, TSharedPtr<FConcertLogFilter> InOptionalFilter);
+	~FFilteredConcertLogList();
+	
+	DECLARE_MULTICAST_DELEGATE_OneParam(FOnLogListChanged, const TArray<TSharedPtr<FConcertLogEntry>>& /*NewFilteredView*/);
+	FOnLogListChanged& OnLogListChanged() { return OnLogListChangedEvent; }
+
+	bool IsBufferFull() const { return LogSource->IsBufferFull(); }
+	const TArray<TSharedPtr<FConcertLogEntry>>& GetFilteredLogs() const { return FilteredLogs; }
+
+protected:
+
+	const TSharedRef<IConcertLogSource>& GetLogSource() const { return LogSource; }
+	
+private:
+
+	/** Determines where the logs come from */
+	TSharedRef<IConcertLogSource> LogSource;
+
+	/** Optional. Filters items from the LogSource */
+	TSharedPtr<FConcertLogFilter> Filter;
+
+	/** Contains the items of LogSource with Filter applied. */
+	TArray<TSharedPtr<FConcertLogEntry>> FilteredLogs;
+	
+	/** Called when the content of FilteredResult changes */
+	FOnLogListChanged OnLogListChangedEvent;
+
+	void OnLowestLogEntryChanged(FConcertLogID NewLowestValidID);
+	void OnNewLogEntryAdded(const TSharedRef<FConcertLogEntry>& NewLogEntry);
+	void RebuildFilteredResult();
+};
+
+/** Further sorts the results of FFilteredConcertLogList::Filter into pages, i.e. buckets of logs, useful for limiting elements in a table view. */
+class FPagedFilteredConcertLogList : public FFilteredConcertLogList
+{
+public:
+	
+	using FLogsPerPageCount = uint16;
+	using FPageCount = uint32;
+
+	FPagedFilteredConcertLogList(TSharedRef<IConcertLogSource> LogSource, TSharedPtr<FConcertLogFilter> OptionalFilter, uint16 FLogsPerPageCount = 500);
+
+	FPageCount GetNumPages() const { const FPageCount PageCount = GetFilteredLogs().Num() / LogsPerPage; return GetFilteredLogs().Num() % LogsPerPage == 0 ? PageCount : PageCount + 1;}
+	FPageCount GetCurrentPage() const { return CurrentPageIndex; }
+	uint16 GetLogsPerPage() const { return LogsPerPage; }
+	const TArray<TSharedPtr<FConcertLogEntry>>& GetPageView() const { return PageView; }
+	
+	void SetLogsPerPage(FLogsPerPageCount NewLogsPerPage);
+	void SetPage(const uint32 PageIndex);
+
+	FOnLogListChanged& OnPageViewChanged() { return OnPageViewChangedEvent; }
+	
+private:
+
+
+	/** The number of logs to show per page*/
+	FLogsPerPageCount LogsPerPage;
+
+	/** The index of the current page. Starts at 0. */
+	FPageCount CurrentPageIndex = 0;
+
+	/** The filtered logs on the current page */
+	TArray<TSharedPtr<FConcertLogEntry>> PageView;
+
+	/** Called when the content of PageView changes */
+	FOnLogListChanged OnPageViewChangedEvent;
+
+	void HandeOnLogListChanged(const TArray<TSharedPtr<FConcertLogEntry>>& NewFilteredLogList);
+	void RepopulatePage();
+	void CheckAndConditionallyPopulatePage();
+	
+	/**
+	 * Given the current page settings (not validated - assumed corrected), calls Callback with each index in GetFilteredLogs()
+	 * that should be an item on the current page. */
+	void ForEachLogIndexOnPage(TFunctionRef<void(size_t FilteredLogIndex)> Callback, FLogsPerPageCount MaxItems = TNumericLimits<FLogsPerPageCount>::Max());
+	/** @return Start and last index in GetFilteredLogs() array to accommodate current page settings */
+	TTuple<size_t, size_t> GetLogIndicesForPage() const;
+};
+
diff --git a/Engine/Plugins/Developer/Concert/ConcertApp/MultiUserServer/Source/MultiUserServer/Private/Widgets/Clients/Logging/SConcertTransportLog.cpp b/Engine/Plugins/Developer/Concert/ConcertApp/MultiUserServer/Source/MultiUserServer/Private/Widgets/Clients/Logging/SConcertTransportLog.cpp
new file mode 100644
--- /dev/null
+++ b/Engine/Plugins/Developer/Concert/ConcertApp/MultiUserServer/Source/MultiUserServer/Private/Widgets/Clients/Logging/SConcertTransportLog.cpp
@@ -0,0 +1,230 @@
+// Copyright Epic Games, Inc. All Rights Reserved.
+
+#include "SConcertTransportLog.h"
+
+#include "ConcertHeaderRowUtils.h"
+#include "Filter/ConcertLogFilter_FrontendRoot.h"
+#include "Filter/ConcertFrontendLogFilter_TextSearch.h"
+#include "Util/ConcertLogTokenizer.h"
+#include "SConcertTransportLogFooter.h"
+#include "SConcertTransportLogRow.h"
+#include "Settings/ConcertTransportLogSettings.h"
+#include "Settings/MultiUserServerColumnVisibilitySettings.h"
+
+#include "Styling/AppStyle.h"
+#include "Widgets/Layout/SBorder.h"
+#include "Widgets/SBoxPanel.h"
+#include "Widgets/SNullWidget.h"
+#include "Widgets/Views/SHeaderRow.h"
+#include "Widgets/Views/SListView.h"
+
+#define LOCTEXT_NAMESPACE "UnrealMultiUserUI"
+
+const FName SConcertTransportLog::FirstColumnId("AvatarColourColumnId");
+
+void SConcertTransportLog::Construct(const FArguments& InArgs, TSharedRef<IConcertLogSource> LogSource)
+{
+	PagedLogList = MakeShared<FPagedFilteredConcertLogList>(MoveTemp(LogSource), InArgs._Filter);
+	LogTokenizer = MakeShared<FConcertLogTokenizer>();
+	HighlightText = MakeShared<FText>();
+
+	GetClientInfoFunc = InArgs._GetClientInfo;
+
+	ChildSlot
+	[
+		SNew(SBorder)
+		.BorderImage(FAppStyle::Get().GetBrush("ToolPanel.GroupBorder"))
+		.BorderBackgroundColor(FSlateColor(FLinearColor(0.6, 0.6, 0.6)))
+		.Padding(2)
+		[
+			SNew(SVerticalBox)
+
+			+SVerticalBox::Slot()
+			.AutoHeight()
+			.VAlign(VAlign_Top)
+			[
+				InArgs._Filter ? InArgs._Filter->BuildFilterWidgets() : SNullWidget::NullWidget
+			]
+
+			+SVerticalBox::Slot()
+			.FillHeight(1.f)
+			.Padding(0, 5, 0, 0)
+			[
+				CreateTableView()
+			]
+
+			+SVerticalBox::Slot()
+			.AutoHeight()
+			[
+				SNew(SConcertTransportLogFooter, PagedLogList.ToSharedRef())
+				.ExtendViewOptions(this, &SConcertTransportLog::ExtendViewOptions)
+			]
+		]
+	];
+
+	// Subscribe to events
+	PagedLogList->OnPageViewChanged().AddSP(this, &SConcertTransportLog::OnPageViewChanged);
+	if (InArgs._Filter)
+	{
+		InArgs._Filter->GetTextSearchFilter()->OnSearchTextChanged().AddSP(this, &SConcertTransportLog::OnSearchTextChanged);
+	}
+	UMultiUserServerColumnVisibilitySettings::GetSettings()->OnTransportLogColumnVisibility().AddSP(this, &SConcertTransportLog::OnColumnVisibilitySettingsChanged);
+	
+	UE::ConcertSharedSlate::RestoreColumnVisibilityState(HeaderRow.ToSharedRef(), UMultiUserServerColumnVisibilitySettings::GetSettings()->GetTransportLogColumnVisibility());
+}
+
+TSharedRef<SWidget> SConcertTransportLog::CreateTableView()
+{
+	return SAssignNew(LogView, SListView<TSharedPtr<FConcertLogEntry>>)
+		.ListItemsSource(&PagedLogList->GetPageView())
+		.OnGenerateRow(this, &SConcertTransportLog::OnGenerateActivityRowWidget)
+		.SelectionMode(ESelectionMode::None)
+		.HeaderRow(CreateHeaderRow());
+}
+
+TSharedRef<SHeaderRow> SConcertTransportLog::CreateHeaderRow()
+{
+	HeaderRow = SNew(SHeaderRow)
+		.OnHiddenColumnsListChanged_Lambda([this]()
+		{
+			if (!bIsUpdatingColumnVisibility)
+			{
+				UMultiUserServerColumnVisibilitySettings::GetSettings()->SetTransportLogColumnVisibility(
+				UE::ConcertSharedSlate::SnapshotColumnVisibilityState(HeaderRow.ToSharedRef())
+				);
+			}
+		})
+		// Create tiny dummy row showing avatar colour to handle the case when a user hides all columns
+		+SHeaderRow::Column(FirstColumnId)
+			.DefaultLabel(FText::GetEmpty())
+			.FixedWidth(8)
+			// Cannot be hidden
+			.ShouldGenerateWidget(true)
+			.ToolTipText(LOCTEXT("AvatarColumnToolTipText", "The colour of the avatar is affected by log"))
+		;
+
+	// The following properties are useless to the user
+	const TSet<const FProperty*> SkippedProperties = {
+		// Used by code only
+		FConcertLog::StaticStruct()->FindPropertyByName(GET_MEMBER_NAME_CHECKED(FConcertLog, CustomPayloadUncompressedByteSize))
+	};
+	
+	for (TFieldIterator<const FProperty> PropertyIt(FConcertLog::StaticStruct()); PropertyIt; ++PropertyIt)
+	{
+		if (!PropertyIt->HasAnyPropertyFlags(CPF_Transient) && !SkippedProperties.Contains(*PropertyIt))
+		{
+			const FString PropertyName = PropertyIt->GetAuthoredName();
+			const FName ColumnId = PropertyIt->GetFName();
+			HeaderRow->AddColumn(SHeaderRow::FColumn::FArguments()
+				.ColumnId(ColumnId)
+				.DefaultLabel(FText::FromString(PropertyName))
+				.HAlignCell(HAlign_Center)
+				// Add option to hide
+				.OnGetMenuContent_Lambda([this, ColumnId]()
+				{
+					return UE::ConcertSharedSlate::MakeHideColumnContextMenu(HeaderRow.ToSharedRef(), ColumnId);
+				}));
+		}
+	}
+
+	TGuardValue<bool> DoNotSave(bIsUpdatingColumnVisibility, true);
+	const TSet<FName> HiddenByDefault = {
+		GET_MEMBER_NAME_CHECKED(FConcertLog, Frame),
+		GET_MEMBER_NAME_CHECKED(FConcertLog, MessageId),
+		GET_MEMBER_NAME_CHECKED(FConcertLog, MessageOrderIndex),
+		GET_MEMBER_NAME_CHECKED(FConcertLog, ChannelId),
+		GET_MEMBER_NAME_CHECKED(FConcertLog, CustomPayloadTypename),
+		GET_MEMBER_NAME_CHECKED(FConcertLog, StringPayload)
+	};
+	for (const FName ColumnID : HiddenByDefault)
+	{
+		HeaderRow->SetShowGeneratedColumn(ColumnID, false);
+	}
+	
+	return HeaderRow.ToSharedRef();
+}
+
+TSharedRef<ITableRow> SConcertTransportLog::OnGenerateActivityRowWidget(TSharedPtr<FConcertLogEntry> Item, const TSharedRef<STableViewBase>& OwnerTable) const
+{
+	return SNew(SConcertTransportLogRow, Item, OwnerTable, LogTokenizer.ToSharedRef(), HighlightText.ToSharedRef())
+		.GetClientInfo(GetClientInfoFunc);
+}
+
+void SConcertTransportLog::ExtendViewOptions(FMenuBuilder& MenuBuilder)
+{
+	MenuBuilder.AddMenuEntry(
+		LOCTEXT("AutoScroll", "Auto Scroll"),
+		LOCTEXT("AutoScroll_Tooltip", "Automatically scroll as new logs arrive (affects last page)"),
+		FSlateIcon(),
+		FUIAction(
+			FExecuteAction::CreateLambda([this](){ bAutoScroll = !bAutoScroll; }),
+			FCanExecuteAction::CreateLambda([] { return true; }),
+			FIsActionChecked::CreateLambda([this] { return bAutoScroll; })),
+		NAME_None,
+		EUserInterfaceActionType::ToggleButton
+	);
+
+	MenuBuilder.AddMenuEntry(
+			LOCTEXT("DisplayTimestampInRelativeTime", "Display Relative Time"),
+			TAttribute<FText>::CreateLambda([this]()
+			{
+				const bool bIsVisible = HeaderRow->IsColumnVisible(GET_MEMBER_NAME_CHECKED(FConcertLog, Timestamp));
+				return bIsVisible
+					? LOCTEXT("DisplayTimestampInRelativeTime.Tooltip.Visible", "Display the Last Modified column in relative time?")
+					: LOCTEXT("DisplayTimestampInRelativeTime.Tooltip.Hidden", "Disabled because the Timestamp column is hidden.");
+			}),
+			FSlateIcon(),
+			FUIAction(
+				FExecuteAction::CreateSP(this, &SConcertTransportLog::OnFilterMenuChecked),
+				FCanExecuteAction::CreateLambda([this] { return HeaderRow->IsColumnVisible(GET_MEMBER_NAME_CHECKED(FConcertLog, Timestamp)); }),
+				FIsActionChecked::CreateLambda([this] { return UConcertTransportLogSettings::GetSettings()->TimestampTimeFormat == ETimeFormat::Relative; })),
+			NAME_None,
+			EUserInterfaceActionType::ToggleButton
+		);
+
+					
+	MenuBuilder.AddSeparator();
+	UE::ConcertSharedSlate::AddEntriesForShowingHiddenRows(HeaderRow.ToSharedRef(), MenuBuilder);
+}
+
+void SConcertTransportLog::OnFilterMenuChecked()
+{
+	UConcertTransportLogSettings* Settings = UConcertTransportLogSettings::GetSettings();
+	
+	switch (Settings->TimestampTimeFormat)
+	{
+	case ETimeFormat::Relative:
+		Settings->TimestampTimeFormat = ETimeFormat::Absolute;
+		break;
+	case ETimeFormat::Absolute: 
+		Settings->TimestampTimeFormat = ETimeFormat::Relative;
+		break;
+	default:
+		checkNoEntry();
+	}
+
+	Settings->SaveConfig();
+}
+
+void SConcertTransportLog::OnPageViewChanged(const TArray<TSharedPtr<FConcertLogEntry>>&)
+{
+	LogView->RequestListRefresh();
+
+	if (bAutoScroll)
+	{
+		LogView->ScrollToBottom();
+	}
+}
+
+void SConcertTransportLog::OnSearchTextChanged(const FText& NewSearchText)
+{
+	*HighlightText = NewSearchText;
+}
+
+void SConcertTransportLog::OnColumnVisibilitySettingsChanged(const FColumnVisibilitySnapshot& ColumnSnapshot)
+{
+	TGuardValue<bool> GuardValue(bIsUpdatingColumnVisibility, true);
+	UE::ConcertSharedSlate::RestoreColumnVisibilityState(HeaderRow.ToSharedRef(), ColumnSnapshot);
+}
+
+#undef LOCTEXT_NAMESPACE
\ No newline at end of file
diff --git a/Engine/Plugins/Developer/Concert/ConcertApp/MultiUserServer/Source/MultiUserServer/Private/Widgets/Clients/Logging/SConcertTransportLog.h b/Engine/Plugins/Developer/Concert/ConcertApp/MultiUserServer/Source/MultiUserServer/Private/Widgets/Clients/Logging/SConcertTransportLog.h
new file mode 100644
--- /dev/null
+++ b/Engine/Plugins/Developer/Concert/ConcertApp/MultiUserServer/Source/MultiUserServer/Private/Widgets/Clients/Logging/SConcertTransportLog.h
@@ -0,0 +1,73 @@
+// Copyright Epic Games, Inc. All Rights Reserved.
+
+#pragma once
+
+#include "CoreMinimal.h"
+#include "Filter/FilteredConcertLogList.h"
+#include "TransportLogDelegates.h"
+#include "Widgets/DeclarativeSyntaxSupport.h"
+#include "Widgets/SCompoundWidget.h"
+
+struct FColumnVisibilitySnapshot;
+class FConcertLogTokenizer;
+class FConcertLogFilter_FrontendRoot;
+class FMenuBuilder;
+class FPagedFilteredConcertLogList;
+class IConcertLogSource;
+template <typename ItemType> class SListView;
+
+struct FConcertLogEntry;
+
+/**
+ * Displays the contents of a IConcertLogSource and has UI for filtering.
+ */
+class SConcertTransportLog : public SCompoundWidget
+{
+public:
+
+	static const FName FirstColumnId;
+
+	SLATE_BEGIN_ARGS(SConcertTransportLog)
+	{}
+		/** Optional filters to display in UI */
+		SLATE_ARGUMENT(TSharedPtr<FConcertLogFilter_FrontendRoot>, Filter)
+		/** Used to better display client info, e.g. client name and colour */
+		SLATE_EVENT(FGetClientInfo, GetClientInfo)
+	SLATE_END_ARGS()
+
+	void Construct(const FArguments& InArgs, TSharedRef<IConcertLogSource> LogSource);
+
+private:
+
+	/** Sorts the log into pages whilst applying filters */
+	TSharedPtr<FPagedFilteredConcertLogList> PagedLogList;
+	/** Used by various systems to convert logs to text */
+	TSharedPtr<FConcertLogTokenizer> LogTokenizer;
+
+	/** Updates to be the content of the search text. Shared with all rows. */
+	TSharedPtr<FText> HighlightText;
+	
+	/** Lists the logs */
+	TSharedPtr<SListView<TSharedPtr<FConcertLogEntry>>> LogView;
+	/** Header row of LogView */
+	TSharedPtr<SHeaderRow> HeaderRow;
+
+	FGetClientInfo GetClientInfoFunc;
+
+	/** Whether to automatically scroll to new logs as they come in */
+	bool bAutoScroll = true;
+	/** Whether we currently loading the column visibility - prevents infinite event recursion */
+	bool bIsUpdatingColumnVisibility = false;
+
+	// Table view creation
+	TSharedRef<SWidget> CreateTableView();
+	TSharedRef<SHeaderRow> CreateHeaderRow();
+	TSharedRef<ITableRow> OnGenerateActivityRowWidget(TSharedPtr<FConcertLogEntry> Item, const TSharedRef<STableViewBase>& OwnerTable) const;
+
+	void ExtendViewOptions(FMenuBuilder& MenuBuilder);
+	void OnFilterMenuChecked();
+	void OnPageViewChanged(const TArray<TSharedPtr<FConcertLogEntry>>&);
+	void OnSearchTextChanged(const FText& NewSearchText);
+	void OnColumnVisibilitySettingsChanged(const FColumnVisibilitySnapshot& ColumnSnapshot);
+};
+
diff --git a/Engine/Plugins/Developer/Concert/ConcertApp/MultiUserServer/Source/MultiUserServer/Private/Widgets/Clients/Logging/SConcertTransportLogFooter.cpp b/Engine/Plugins/Developer/Concert/ConcertApp/MultiUserServer/Source/MultiUserServer/Private/Widgets/Clients/Logging/SConcertTransportLogFooter.cpp
new file mode 100644
--- /dev/null
+++ b/Engine/Plugins/Developer/Concert/ConcertApp/MultiUserServer/Source/MultiUserServer/Private/Widgets/Clients/Logging/SConcertTransportLogFooter.cpp
@@ -0,0 +1,148 @@
+// Copyright Epic Games, Inc. All Rights Reserved.
+
+#include "SConcertTransportLogFooter.h"
+
+#include "Filter/FilteredConcertLogList.h"
+#include "Styling/AppStyle.h"
+#include "Widgets/SBoxPanel.h"
+#include "Widgets/Input/SComboButton.h"
+#include "Widgets/Input/SSpinBox.h"
+#include "Widgets/Text/STextBlock.h"
+
+#define LOCTEXT_NAMESPACE "UnrealMultiUserUI"
+
+namespace UE::MultiUserServer::Private
+{
+	constexpr size_t DefaultPageSizeIndex = 3; 
+	static TArray<TSharedPtr<FPagedFilteredConcertLogList::FLogsPerPageCount>> PageSizeOptions = {
+		MakeShared<FPagedFilteredConcertLogList::FLogsPerPageCount>(10),
+		MakeShared<FPagedFilteredConcertLogList::FLogsPerPageCount>(50),
+		MakeShared<FPagedFilteredConcertLogList::FLogsPerPageCount>(100),
+		MakeShared<FPagedFilteredConcertLogList::FLogsPerPageCount>(500),
+		MakeShared<FPagedFilteredConcertLogList::FLogsPerPageCount>(1000),
+		MakeShared<FPagedFilteredConcertLogList::FLogsPerPageCount>(5000)
+	};
+}
+
+void SConcertTransportLogFooter::Construct(const FArguments& InArgs, TSharedRef<FPagedFilteredConcertLogList> InPagedLogList)
+{
+	using namespace UE::MultiUserServer::Private;
+	PagedLogList = MoveTemp(InPagedLogList);
+
+	ChildSlot
+	[
+		SNew(SHorizontalBox)
+		
+		// Page selection
+		+SHorizontalBox::Slot()
+		.AutoWidth()
+		.VAlign(VAlign_Center)
+		[
+			SNew(SSpinBox<FPagedFilteredConcertLogList::FPageCount>)
+			.Value_Lambda([this](){ return PagedLogList->GetFilteredLogs().Num() == 0 ? 0 : PagedLogList->GetCurrentPage() + 1; })
+			.OnValueChanged_Lambda([this](FPagedFilteredConcertLogList::FPageCount NewValue) { PagedLogList->SetPage(NewValue - 1); })
+			.MinValue(0)
+			.MaxValue_Lambda([this]() { return PagedLogList->GetNumPages(); })
+		]
+
+		// Page counter
+		+SHorizontalBox::Slot()
+		.AutoWidth()
+		.Padding(2, 0, 0, 0)
+		.VAlign(VAlign_Center)
+		[
+			SNew(STextBlock)
+			.Text_Lambda([this](){ return FText::Format(LOCTEXT("PageCounterFormat", "of {0} {0}|plural(one=page,other=pages)"), PagedLogList->GetNumPages()); })
+		]
+
+		// Logs per page text
+		+SHorizontalBox::Slot()
+		.AutoWidth()
+		.Padding(30, 0, 0, 0)
+		.VAlign(VAlign_Center)
+		[
+			SNew(STextBlock)
+			.Text(LOCTEXT("LogePerPageText", "Logs per page"))
+		]
+		
+		// Logs per page combo box
+		+SHorizontalBox::Slot()
+		.AutoWidth()
+		.Padding(3, 0, 0, 0)
+		.VAlign(VAlign_Center)
+		[
+			SNew(SComboBox<TSharedPtr<FPagedFilteredConcertLogList::FLogsPerPageCount>>)
+			.InitiallySelectedItem(PageSizeOptions[DefaultPageSizeIndex])
+			.OnSelectionChanged_Lambda([this](const TSharedPtr<FPagedFilteredConcertLogList::FLogsPerPageCount>& NewCount, ESelectInfo::Type){ PagedLogList->SetLogsPerPage(*NewCount); })
+			.OnGenerateWidget_Lambda([](const TSharedPtr<FPagedFilteredConcertLogList::FLogsPerPageCount>& Count) { return SNew(STextBlock).Text(FText::FromString(FString::FromInt(*Count))); })
+			.OptionsSource(&PageSizeOptions)
+			[
+				SNew(STextBlock).Text_Lambda([this](){ return FText::FromString(FString::FromInt(PagedLogList->GetLogsPerPage())); })
+			]
+		]
+
+		// Log counter
+		+SHorizontalBox::Slot()
+		.AutoWidth()
+		.Padding(30, 0, 0, 0)
+		.VAlign(VAlign_Center)
+		[
+			SNew(STextBlock)
+			.Text_Lambda([this]()
+			{
+				return PagedLogList->IsBufferFull()
+					? FText::Format(LOCTEXT("LogCounterFormat.Full", "{0} filtered {0}|plural(one=log,other=logs) (buffer full)"), PagedLogList->GetFilteredLogs().Num())
+					: FText::Format(LOCTEXT("LogCounterFormat.NotFull", "{0} filtered {0}|plural(one=log,other=logs)"), PagedLogList->GetFilteredLogs().Num());
+			})
+		]
+
+		// Gap Filler
+		+SHorizontalBox::Slot()
+		.FillWidth(1.0)
+		[
+			SNew(SSpacer)
+		]
+
+		+SHorizontalBox::Slot()
+		.AutoWidth()
+		.VAlign(VAlign_Center)
+		[
+			SNew(SComboButton)
+				.ComboButtonStyle(FAppStyle::Get(), "GenericFilters.ComboButtonStyle")
+				.ForegroundColor(FLinearColor::White)
+				.ContentPadding(0)
+				.OnGetMenuContent(this, &SConcertTransportLogFooter::MakeViewOptionsMenuWidget, InArgs._ExtendViewOptions)
+				.HasDownArrow(true)
+				.ContentPadding(FMargin(1, 0))
+				.ButtonContent()
+				[
+					SNew(SHorizontalBox)
+
+					+SHorizontalBox::Slot()
+					.AutoWidth()
+					.VAlign(VAlign_Center)
+					[
+						SNew(SImage).Image(FAppStyle::Get().GetBrush("Icons.Visible")) // The eye ball image.
+					]
+
+					+SHorizontalBox::Slot()
+					.AutoWidth()
+					.Padding(2, 0, 0, 0)
+					.VAlign(VAlign_Center)
+					[
+						SNew(STextBlock).Text(LOCTEXT("ViewOptions", "View Options"))
+					]
+				]
+		]
+	];
+}
+
+TSharedRef<SWidget> SConcertTransportLogFooter::MakeViewOptionsMenuWidget(FExtendContextMenu ExtendViewOptions)
+{
+	FMenuBuilder MenuBuilder(/*bInShouldCloseWindowAfterMenuSelection=*/true, nullptr);
+
+	ExtendViewOptions.Execute(MenuBuilder);
+	return MenuBuilder.MakeWidget();
+}
+
+#undef LOCTEXT_NAMESPACE
diff --git a/Engine/Plugins/Developer/Concert/ConcertApp/MultiUserServer/Source/MultiUserServer/Private/Widgets/Clients/Logging/SConcertTransportLogFooter.h b/Engine/Plugins/Developer/Concert/ConcertApp/MultiUserServer/Source/MultiUserServer/Private/Widgets/Clients/Logging/SConcertTransportLogFooter.h
new file mode 100644
--- /dev/null
+++ b/Engine/Plugins/Developer/Concert/ConcertApp/MultiUserServer/Source/MultiUserServer/Private/Widgets/Clients/Logging/SConcertTransportLogFooter.h
@@ -0,0 +1,37 @@
+// Copyright Epic Games, Inc. All Rights Reserved.
+
+#pragma once
+
+#include "CoreMinimal.h"
+#include "Filter/FilteredConcertLogList.h"
+#include "Widgets/DeclarativeSyntaxSupport.h"
+#include "Widgets/SCompoundWidget.h"
+
+template<typename T> class SSpinBox;
+struct FConcertLogEntry;
+
+/** Displays the number of pages and items. Displayed under the table view. */
+class SConcertTransportLogFooter : public SCompoundWidget
+{
+public:
+	
+	DECLARE_DELEGATE_OneParam(FExtendContextMenu, FMenuBuilder&)
+
+	SLATE_BEGIN_ARGS(SConcertTransportLogFooter) {}
+		SLATE_EVENT(FExtendContextMenu, ExtendViewOptions)
+	SLATE_END_ARGS()
+
+	void Construct(const FArguments& InArgs, TSharedRef<FPagedFilteredConcertLogList> InPagedLogList);
+
+private:
+
+	/** The model we'll be updating */
+	TSharedPtr<FPagedFilteredConcertLogList> PagedLogList;
+
+	/** Selects the current page */
+	TSharedPtr<SSpinBox<FPagedFilteredConcertLogList::FPageCount>> CurrentPage;
+	/** Displays the number of pages */
+	TSharedPtr<STextBlock> PageCounterText;
+
+	TSharedRef<SWidget> MakeViewOptionsMenuWidget(FExtendContextMenu ExtendViewOptions);
+};
diff --git a/Engine/Plugins/Developer/Concert/ConcertApp/MultiUserServer/Source/MultiUserServer/Private/Widgets/Clients/Logging/SConcertTransportLogRow.cpp b/Engine/Plugins/Developer/Concert/ConcertApp/MultiUserServer/Source/MultiUserServer/Private/Widgets/Clients/Logging/SConcertTransportLogRow.cpp
new file mode 100644
--- /dev/null
+++ b/Engine/Plugins/Developer/Concert/ConcertApp/MultiUserServer/Source/MultiUserServer/Private/Widgets/Clients/Logging/SConcertTransportLogRow.cpp
@@ -0,0 +1,106 @@
+// Copyright Epic Games, Inc. All Rights Reserved.
+
+#include "SConcertTransportLogRow.h"
+
+#include "Util/ConcertLogTokenizer.h"
+#include "SConcertTransportLog.h"
+#include "Session/Activity/PredefinedActivityColumns.h"
+
+#include "Widgets/Colors/SColorBlock.h"
+#include "Widgets/Text/STextBlock.h"
+
+void SConcertTransportLogRow::Construct(const FArguments& InArgs, TSharedPtr<FConcertLogEntry> InLogEntry, const TSharedRef<STableViewBase>& InOwnerTableView, TSharedRef<FConcertLogTokenizer> InTokenizer, TSharedRef<FText> InHighlightText)
+{
+	LogEntry = MoveTemp(InLogEntry);
+	Tokenizer = MoveTemp(InTokenizer);
+	HighlightText = MoveTemp(InHighlightText);
+
+	GetClientInfoFunc = InArgs._GetClientInfo;
+	
+	SMultiColumnTableRow<TSharedPtr<FConcertLogEntry>>::Construct(FSuperRowType::FArguments(), InOwnerTableView);
+}
+
+TSharedRef<SWidget> SConcertTransportLogRow::GenerateWidgetForColumn(const FName& ColumnName)
+{
+	if (ColumnName == SConcertTransportLog::FirstColumnId)
+	{
+		return SNew(SHorizontalBox)
+			+SHorizontalBox::Slot()
+			.Padding(2, 1)
+			[
+				SNew(SColorBlock)
+				.Color_Lambda([this]()
+				{
+					const TOptional<FConcertClientInfo> ClientInfo = GetClientInfoFromLog();
+					return ClientInfo.IsSet()
+						? ClientInfo->AvatarColor
+						: FLinearColor::Black;
+				})
+				.Size(FVector2D(4.f, 16.f))
+			];
+	}
+	
+	using FColumnWidgetFactoryFunc = TSharedRef<SWidget>(SConcertTransportLogRow::*)(const FName& InColumnName);
+	const TMap<FName, FColumnWidgetFactoryFunc> OverrideFactories = {
+		{ GET_MEMBER_NAME_CHECKED(FConcertLog, OriginEndpointId), &SConcertTransportLogRow::CreateClientNameColumn },
+		{ GET_MEMBER_NAME_CHECKED(FConcertLog, DestinationEndpointId), &SConcertTransportLogRow::CreateClientNameColumn }
+	};
+	const FColumnWidgetFactoryFunc FallbackFactoryFunc = &SConcertTransportLogRow::CreateDefaultColumn;
+
+	const FColumnWidgetFactoryFunc* FactoryFunc = OverrideFactories.Find(ColumnName);
+	return FactoryFunc
+		? Invoke(*FactoryFunc, this, ColumnName)
+		: Invoke(FallbackFactoryFunc, this, ColumnName);
+}
+
+TSharedRef<SWidget> SConcertTransportLogRow::CreateDefaultColumn(const FName& PropertyName)
+{
+	return SNew(STextBlock)
+		.Text_Lambda([PropertyName, WeakEntry = TWeakPtr<FConcertLogEntry>(LogEntry), WeakTokenizer = TWeakPtr<FConcertLogTokenizer>(Tokenizer)]()
+		{
+			TSharedPtr<FConcertLogEntry> PinnedEntry = WeakEntry.Pin();
+			TSharedPtr<FConcertLogTokenizer> PinnedTokenizer = WeakTokenizer.Pin();
+			const FProperty* Property = FConcertLog::StaticStruct()->FindPropertyByName(PropertyName);
+			if (ensure(PinnedEntry && PinnedTokenizer && Property))
+			{
+				return FText::FromString(PinnedTokenizer->Tokenize(PinnedEntry->Log, *Property));
+			}
+				
+			return FText::GetEmpty();
+		})
+		.HighlightText_Lambda([this](){ return *HighlightText.Get(); });
+}
+
+TSharedRef<SWidget> SConcertTransportLogRow::CreateClientNameColumn(const FName& PropertyName)
+{
+	const bool bUseOrigin = GET_MEMBER_NAME_CHECKED(FConcertLog, OriginEndpointId) == PropertyName;
+	const bool bUseDestination = GET_MEMBER_NAME_CHECKED(FConcertLog, DestinationEndpointId) == PropertyName;
+	check(bUseOrigin || bUseDestination);
+
+	const FGuid& EndpointId = bUseOrigin
+		? LogEntry->Log.OriginEndpointId
+		: LogEntry->Log.DestinationEndpointId;
+	const TOptional<FConcertClientInfo> ClientInfo = GetClientInfoFunc.IsBound()
+		? GetClientInfoFunc.Execute(EndpointId)
+		: TOptional<FConcertClientInfo>();
+	return SNew(STextBlock)
+		// If there is no client info available when the log is
+		.Text(FText::FromString(ClientInfo.IsSet() ? ClientInfo->DisplayName : EndpointId.ToString()))
+		.HighlightText_Lambda([this](){ return *HighlightText.Get(); });
+}
+
+TOptional<FConcertClientInfo> SConcertTransportLogRow::GetClientInfoFromLog() const
+{
+	if (!GetClientInfoFunc.IsBound())
+	{
+		return {};
+	}
+		
+	const TOptional<FConcertClientInfo> OriginClientInfo = GetClientInfoFunc.IsBound() ? GetClientInfoFunc.Execute(LogEntry->Log.OriginEndpointId) : TOptional<FConcertClientInfo>();
+	const TOptional<FConcertClientInfo> DestinationClientInfo = GetClientInfoFunc.IsBound() ? GetClientInfoFunc.Execute(LogEntry->Log.DestinationEndpointId) : TOptional<FConcertClientInfo>();
+
+	// One of the two is definitely the server. If one of them is a client whose info could be found, this will return it.
+	return OriginClientInfo.IsSet()
+		? OriginClientInfo
+		: DestinationClientInfo;
+}
diff --git a/Engine/Plugins/Developer/Concert/ConcertApp/MultiUserServer/Source/MultiUserServer/Private/Widgets/Clients/Logging/SConcertTransportLogRow.h b/Engine/Plugins/Developer/Concert/ConcertApp/MultiUserServer/Source/MultiUserServer/Private/Widgets/Clients/Logging/SConcertTransportLogRow.h
new file mode 100644
--- /dev/null
+++ b/Engine/Plugins/Developer/Concert/ConcertApp/MultiUserServer/Source/MultiUserServer/Private/Widgets/Clients/Logging/SConcertTransportLogRow.h
@@ -0,0 +1,43 @@
+// Copyright Epic Games, Inc. All Rights Reserved.
+
+#pragma once
+
+#include "CoreMinimal.h"
+#include "ConcertLogEntry.h"
+#include "TransportLogDelegates.h"
+#include "Widgets/DeclarativeSyntaxSupport.h"
+
+class FConcertLogTokenizer;
+
+class SConcertTransportLogRow : public SMultiColumnTableRow<TSharedPtr<FConcertLogEntry>>
+{
+public:
+
+	SLATE_BEGIN_ARGS(SConcertTransportLogRow)
+	{}
+		/** Used to get the avatar's colour */
+		SLATE_EVENT(FGetClientInfo, GetClientInfo)
+	SLATE_END_ARGS()
+
+	void Construct(const FArguments& InArgs, TSharedPtr<FConcertLogEntry> InLogEntry, const TSharedRef<STableViewBase>& InOwnerTableView, TSharedRef<FConcertLogTokenizer> InTokenizer, TSharedRef<FText> InHighlightText);
+
+	//~ Begin SMultiColumnTableRow Interface
+	virtual TSharedRef<SWidget> GenerateWidgetForColumn(const FName& ColumnName) override;
+	//~ End SMultiColumnTableRow Interface
+	
+private:
+
+	TSharedPtr<FConcertLogEntry> LogEntry;
+	/** Used to convert some members into display strings */
+	TSharedPtr<FConcertLogTokenizer> Tokenizer;
+
+	/** Owned by SConcertTransportLog. Updated with the search text. */
+	TSharedPtr<FText> HighlightText;
+	
+	FGetClientInfo GetClientInfoFunc;
+
+	TSharedRef<SWidget> CreateDefaultColumn(const FName& PropertyName);
+	TSharedRef<SWidget> CreateClientNameColumn(const FName& PropertyName);
+	
+	TOptional<FConcertClientInfo> GetClientInfoFromLog() const;
+};
diff --git a/Engine/Plugins/Developer/Concert/ConcertApp/MultiUserServer/Source/MultiUserServer/Private/Widgets/Clients/Logging/Source/GlobalLogSource.cpp b/Engine/Plugins/Developer/Concert/ConcertApp/MultiUserServer/Source/MultiUserServer/Private/Widgets/Clients/Logging/Source/GlobalLogSource.cpp
new file mode 100644
--- /dev/null
+++ b/Engine/Plugins/Developer/Concert/ConcertApp/MultiUserServer/Source/MultiUserServer/Private/Widgets/Clients/Logging/Source/GlobalLogSource.cpp
@@ -0,0 +1,71 @@
+// Copyright Epic Games, Inc. All Rights Reserved.
+
+#include "GlobalLogSource.h"
+
+#include "ConcertTransportEvents.h"
+
+#include "Math/NumericLimits.h"
+
+namespace UE::MultiUserServer::Private
+{
+	static constexpr FConcertLogID InvalidID = TNumericLimits<FConcertLogID>::Max();
+}
+
+FGlobalLogSource::FGlobalLogSource(size_t LogCapacity)
+	: LogPtrs(TCircularBuffer<TSharedPtr<FConcertLogEntry>>(LogCapacity))
+	, LogEntries(TCircularBuffer<FConcertLogEntry>(LogCapacity))
+{
+	checkf(LogCapacity > 2, TEXT("Minimum log capacity must be 2 for AddLog's invariant for the lowest new ID to hold"));
+
+	for (size_t i = 0; i < LogPtrs.Capacity(); ++i)
+	{
+		LogPtrs[i] = TSharedPtr<FConcertLogEntry>(
+			// Our buffer holds the memory
+			&LogEntries[i],
+			// No additional work is needed to "free" the memory
+			[](auto){}
+			);
+		
+		LogEntries[i].LogId = UE::MultiUserServer::Private::InvalidID;
+	}
+
+	ConcertTransportEvents::OnConcertServerLogEvent().AddRaw(this, &FGlobalLogSource::OnLogProduced);
+}
+
+FGlobalLogSource::~FGlobalLogSource()
+{
+	ConcertTransportEvents::OnConcertServerLogEvent().RemoveAll(this);
+}
+
+void FGlobalLogSource::OnLogProduced(const IConcertServer&, const FConcertLog& Log)
+{
+	check(IsInGameThread());
+	AddLog(Log);
+}
+
+void FGlobalLogSource::AddLog(const FConcertLog& NewLog)
+{
+	const TSharedRef<FConcertLogEntry> NewEntry = LogPtrs[NextIndexToAddTo].ToSharedRef();
+
+	const bool bNeedToOverrideOldEntry = NewEntry->LogId != UE::MultiUserServer::Private::InvalidID; 
+	if (bNeedToOverrideOldEntry)
+	{
+		// We've wrapped around all elements: we've just reached the smallest because IDs grow strongly monotonously.
+		// Since there are at least two entries (see constructor) and because of the above, the next element has new lowest ID.
+		OldestLogEntryIndex = NextIndexToAddTo;
+		OnLowestLogEntryChanged().Broadcast(NewEntry->LogId + 1);
+		// Everybody that had been the log pointer now knows it is pointing to a new log
+	}
+
+	NewEntry->LogId = NextLogID;
+	NewEntry->Log = NewLog;
+	
+	NextIndexToAddTo = LogPtrs.GetNextIndex(NextIndexToAddTo);
+	++NextLogID;
+	NumberOfLogs = FMath::Min<size_t>(NumberOfLogs + 1, LogPtrs.Capacity());
+	
+	// Could happen if we left it running for 550000 years at 1 million logs per second
+	checkf(NextLogID != UE::MultiUserServer::Private::InvalidID, TEXT("Reached max possible log ID"));
+	
+	OnLogEntryAdded().Broadcast(NewEntry);
+}
diff --git a/Engine/Plugins/Developer/Concert/ConcertApp/MultiUserServer/Source/MultiUserServer/Private/Widgets/Clients/Logging/Source/GlobalLogSource.h b/Engine/Plugins/Developer/Concert/ConcertApp/MultiUserServer/Source/MultiUserServer/Private/Widgets/Clients/Logging/Source/GlobalLogSource.h
new file mode 100644
--- /dev/null
+++ b/Engine/Plugins/Developer/Concert/ConcertApp/MultiUserServer/Source/MultiUserServer/Private/Widgets/Clients/Logging/Source/GlobalLogSource.h
@@ -0,0 +1,54 @@
+// Copyright Epic Games, Inc. All Rights Reserved.
+
+#pragma once
+
+#include "CoreMinimal.h"
+#include "IConcertLogSource.h"
+#include "Containers/CircularBuffer.h"
+
+class IConcertServer;
+
+/**
+ * Listens for log events using ConcertTransportEvents::OnConcertServerLogEvent and buffers them.
+ *
+ * The shared pointers and log entries are allocated in two sequential buffers to improve memory footprint and cache
+ * coherency; the footprint is improved because two big chunks of memory is better than many tiny chunks.
+ */
+class FGlobalLogSource : public IConcertLogSource
+{
+public:
+	
+	FGlobalLogSource(size_t LogCapacity);
+	virtual ~FGlobalLogSource() override;
+
+	//~ Begin IConcertLogSource Interface
+	virtual bool IsBufferFull() const override { return NumberOfLogs == LogPtrs.Capacity(); }
+	virtual size_t GetNumLogs() const override { return NumberOfLogs; }
+	virtual const TSharedPtr<FConcertLogEntry>& GetElement_Unchecked(size_t Index) const override { return LogPtrs[OldestLogEntryIndex + Index]; }
+	virtual FLogEntryIdEvent& OnLowestLogEntryChanged() override { return OnLowestEntryChangedEvent; }
+	virtual FLogEntryEvent& OnLogEntryAdded() override { return OnLogEntryAddedEvent; }
+	//~ End IConcertLogSource Interface
+
+private:
+
+	/** The lowest ID in the buffer */
+	size_t OldestLogEntryIndex = 0;
+	/** The number of logs in the buffer. */
+	size_t NumberOfLogs = 0;
+	/** The next index at which to insert a log entry */
+	size_t NextIndexToAddTo = 0;
+	/** The next log ID to use */
+	FConcertLogID NextLogID = 0;
+	
+	/** Contains all logs. The oldest logs are discarded */
+	TCircularBuffer<TSharedPtr<FConcertLogEntry>> LogPtrs;
+	/** The content that the LogPtrs point to. For better cache coherency. */
+	TCircularBuffer<FConcertLogEntry> LogEntries;
+	// TODO: It would be great to add another TCircularBuffer for the pointers' SharedReferenceCount for a better memory footprint
+
+	FLogEntryIdEvent OnLowestEntryChangedEvent;
+	FLogEntryEvent OnLogEntryAddedEvent;
+
+	void OnLogProduced(const IConcertServer&, const FConcertLog& Log);
+	void AddLog(const FConcertLog& NewLog);
+};
diff --git a/Engine/Plugins/Developer/Concert/ConcertApp/MultiUserServer/Source/MultiUserServer/Private/Widgets/Clients/Logging/Source/IConcertLogSource.h b/Engine/Plugins/Developer/Concert/ConcertApp/MultiUserServer/Source/MultiUserServer/Private/Widgets/Clients/Logging/Source/IConcertLogSource.h
new file mode 100644
--- /dev/null
+++ b/Engine/Plugins/Developer/Concert/ConcertApp/MultiUserServer/Source/MultiUserServer/Private/Widgets/Clients/Logging/Source/IConcertLogSource.h
@@ -0,0 +1,54 @@
+// Copyright Epic Games, Inc. All Rights Reserved.
+
+#pragma once
+
+#include "CoreMinimal.h"
+#include "Widgets/Clients/Logging/ConcertLogEntry.h"
+
+/**
+ * Describes a source of logs - it is similar to an array but with some changes:
+ * - GetElement_Unchecked is used to access elements by index, as usual.
+ * - the highest valid index is NumLogs() - 1, as usual.
+ * - the lowest valid index is not necessarily 0 but is instead LowestLogIndex().
+ *
+ * Implementations can decide to delete the oldest logs if the number of logs grows too high; the implementation will call
+ * OnLowestLogEntryChanged when old logs are removed.
+ */
+class IConcertLogSource
+{
+public:
+	
+	virtual ~IConcertLogSource() = default;
+
+	/** @return Whether the maximum number of logs has been reached */
+	virtual bool IsBufferFull() const = 0;
+	
+	/** @return The number of logs in this source */
+	virtual size_t GetNumLogs() const = 0;
+	/**
+	 * @param Index Satisfies 0 <= Index < NumLogs(). Index 0 is the oldest log.
+	 * @return The log entry at the given Index - crashes if the index is out of bounds.
+	 */
+	virtual const TSharedPtr<FConcertLogEntry>& GetElement_Unchecked(size_t Index) const = 0;
+
+	void ForEachLog(TFunctionRef<void(const TSharedPtr<FConcertLogEntry>& Log)> ConsumerFunc) const
+	{
+		for (size_t i = 0; i < GetNumLogs(); ++i)
+		{
+			ConsumerFunc(GetElement_Unchecked(i));
+		}
+	}
+
+	DECLARE_MULTICAST_DELEGATE_OneParam(FLogEntryIdEvent, FConcertLogID /*LogID*/);
+	/**
+	 * Called when the log buffer is about to remove the oldest logs.
+	 * The event tells the lowest ID that is still valid.
+	 * 
+	 * Once this event has completed, the log data is overwritten with the latest log data.
+	 */
+	virtual FLogEntryIdEvent& OnLowestLogEntryChanged() = 0;
+
+	DECLARE_MULTICAST_DELEGATE_OneParam(FLogEntryEvent, const TSharedRef<FConcertLogEntry>& /*NewEntry*/);
+	/** Called when a new entry has been added to the log buffer. */
+	virtual FLogEntryEvent& OnLogEntryAdded() = 0;
+};
diff --git a/Engine/Plugins/Developer/Concert/ConcertApp/MultiUserServer/Source/MultiUserServer/Private/Widgets/Clients/Logging/TransportLogDelegates.h b/Engine/Plugins/Developer/Concert/ConcertApp/MultiUserServer/Source/MultiUserServer/Private/Widgets/Clients/Logging/TransportLogDelegates.h
new file mode 100644
--- /dev/null
+++ b/Engine/Plugins/Developer/Concert/ConcertApp/MultiUserServer/Source/MultiUserServer/Private/Widgets/Clients/Logging/TransportLogDelegates.h
@@ -0,0 +1,9 @@
+// Copyright Epic Games, Inc. All Rights Reserved.
+
+#pragma once
+
+#include "CoreMinimal.h"
+
+struct FConcertClientInfo;
+
+DECLARE_DELEGATE_RetVal_OneParam(TOptional<FConcertClientInfo>, FGetClientInfo, const FGuid& /*EndpointId*/);
\ No newline at end of file
diff --git a/Engine/Plugins/Developer/Concert/ConcertApp/MultiUserServer/Source/MultiUserServer/Private/Widgets/Clients/Logging/Util/ConcertLogTokenizer.cpp b/Engine/Plugins/Developer/Concert/ConcertApp/MultiUserServer/Source/MultiUserServer/Private/Widgets/Clients/Logging/Util/ConcertLogTokenizer.cpp
new file mode 100644
--- /dev/null
+++ b/Engine/Plugins/Developer/Concert/ConcertApp/MultiUserServer/Source/MultiUserServer/Private/Widgets/Clients/Logging/Util/ConcertLogTokenizer.cpp
@@ -0,0 +1,46 @@
+// Copyright Epic Games, Inc. All Rights Reserved.
+
+#include "ConcertLogTokenizer.h"
+
+#include "ConcertFrontendUtils.h"
+#include "ConcertTransportEvents.h"
+#include "MessageTypeUtils.h"
+#include "Settings/ConcertTransportLogSettings.h"
+
+FConcertLogTokenizer::FConcertLogTokenizer()
+{
+	TokenizerFunctions = {
+			{ FConcertLog::StaticStruct()->FindPropertyByName(GET_MEMBER_NAME_CHECKED(FConcertLog, Timestamp)), [this](const FConcertLog& Log) { return TokenizeTimestamp(Log); } },
+			{ FConcertLog::StaticStruct()->FindPropertyByName(GET_MEMBER_NAME_CHECKED(FConcertLog, MessageTypeName)), [this](const FConcertLog& Log) { return TokenizeMessageTypeName(Log); } }
+		};
+}
+
+FString FConcertLogTokenizer::Tokenize(const FConcertLog& Data, const FProperty& ConcertLogProperty) const
+{
+	if (const FTokenizeFunc* CustomTokenizer = TokenizerFunctions.Find(&ConcertLogProperty))
+	{
+		return (*CustomTokenizer)(Data);
+	}
+	return TokenizeUsingPropertyExport(Data, ConcertLogProperty);
+}
+
+FString FConcertLogTokenizer::TokenizeTimestamp(const FConcertLog& Data) const
+{
+	return ConcertFrontendUtils::FormatTime(Data.Timestamp, UConcertTransportLogSettings::GetSettings()->TimestampTimeFormat).ToString();
+}
+
+FString FConcertLogTokenizer::TokenizeMessageTypeName(const FConcertLog& Data) const
+{
+	return UE::MultiUserServer::MessageTypeUtils::SanitizeMessageTypeName(Data.MessageTypeName);
+}
+
+FString FConcertLogTokenizer::TokenizeUsingPropertyExport(const FConcertLog& Data, const FProperty& ConcertLogProperty) const
+{
+	FString Exported;
+	const void* ValuePtr = ConcertLogProperty.ContainerPtrToValuePtr<void>(&Data);
+	const void* DeltaPtr = ValuePtr; // We have no real delta - in this case the API expects the same ptr
+	const bool bSuccess = ConcertLogProperty.ExportText_Direct(Exported, ValuePtr, DeltaPtr, nullptr, PPF_ExternalEditor);
+	check(bSuccess);
+
+	return Exported;
+}
diff --git a/Engine/Plugins/Developer/Concert/ConcertApp/MultiUserServer/Source/MultiUserServer/Private/Widgets/Clients/Logging/Util/ConcertLogTokenizer.h b/Engine/Plugins/Developer/Concert/ConcertApp/MultiUserServer/Source/MultiUserServer/Private/Widgets/Clients/Logging/Util/ConcertLogTokenizer.h
new file mode 100644
--- /dev/null
+++ b/Engine/Plugins/Developer/Concert/ConcertApp/MultiUserServer/Source/MultiUserServer/Private/Widgets/Clients/Logging/Util/ConcertLogTokenizer.h
@@ -0,0 +1,29 @@
+// Copyright Epic Games, Inc. All Rights Reserved.
+
+#pragma once
+
+#include "CoreMinimal.h"
+
+struct FConcertLog;
+
+/** Converts members of FConcertLog into a string. Used e.g. to make search respect the display settings. */
+class FConcertLogTokenizer
+{
+public:
+
+	FConcertLogTokenizer();
+
+	/** Tokenizes a property of FConcertLog into a string */
+	FString Tokenize(const FConcertLog& Data, const FProperty& ConcertLogProperty) const;
+
+	FString TokenizeTimestamp(const FConcertLog& Data) const;
+	FString TokenizeMessageTypeName(const FConcertLog& Data) const;
+	FString TokenizeUsingPropertyExport(const FConcertLog& Data, const FProperty& ConcertLogProperty) const;
+
+private:
+
+	using FTokenizeFunc = TFunction<FString(const FConcertLog&)>;
+
+	/** Override functions for tokenizing certain properties */
+	TMap<const FProperty*, FTokenizeFunc> TokenizerFunctions;
+};
diff --git a/Engine/Plugins/Developer/Concert/ConcertApp/MultiUserServer/Source/MultiUserServer/Private/Widgets/Clients/Logging/Util/MessageActionUtils.cpp b/Engine/Plugins/Developer/Concert/ConcertApp/MultiUserServer/Source/MultiUserServer/Private/Widgets/Clients/Logging/Util/MessageActionUtils.cpp
new file mode 100644
--- /dev/null
+++ b/Engine/Plugins/Developer/Concert/ConcertApp/MultiUserServer/Source/MultiUserServer/Private/Widgets/Clients/Logging/Util/MessageActionUtils.cpp
@@ -0,0 +1,56 @@
+// Copyright Epic Games, Inc. All Rights Reserved.
+
+#include "MessageActionUtils.h"
+
+#include "Algo/Transform.h"
+
+namespace UE::MultiUserServer::MessageActionUtils
+{
+	TSet<FName> GetAllMessageActionNames()
+	{
+		static const TSet<FName> CachedResult = []()
+		{
+			const TSet<EConcertLogMessageAction> Actions = GetAllMessageActions();
+			TSet<FName> Result;
+			Algo::Transform(Actions, Result, [](EConcertLogMessageAction Action)
+			{
+				return ConvertActionToName(Action);
+			});
+			return Result;
+		}();
+		return CachedResult;
+	}
+
+	TSet<EConcertLogMessageAction> GetAllMessageActions()
+	{
+		constexpr int32 NumActions = 12;
+		static_assert(static_cast<int32>(EConcertLogMessageAction::EndpointClosure) == NumActions - 1, "Update this function when you change EConcertLogMessageAction");
+		static_assert(static_cast<int32>(EConcertLogMessageAction::None) == 0, "Fix the below loop when you change EConcertLogMessageAction");
+
+		TSet<EConcertLogMessageAction> Result;
+		for (int32 i = 0; i < NumActions; ++i)
+		{
+			Result.Add(static_cast<EConcertLogMessageAction>(i));
+		}
+
+		return Result;
+	}
+
+	FName ConvertActionToName(EConcertLogMessageAction MessageAction)
+	{
+		return UEnum::GetValueAsName(MessageAction);
+	}
+
+	FString GetActionDisplayString(EConcertLogMessageAction MessageAction)
+	{
+		const FName Name = ConvertActionToName(MessageAction);
+		return GetActionDisplayString(Name);
+	}
+
+	FString GetActionDisplayString(FName MessageActionName)
+	{
+		FString AsString = MessageActionName.ToString();
+		AsString.RemoveFromStart(TEXT("EConcertLogMessageAction::"));
+		return AsString;
+	}
+}
diff --git a/Engine/Plugins/Developer/Concert/ConcertApp/MultiUserServer/Source/MultiUserServer/Private/Widgets/Clients/Logging/Util/MessageActionUtils.h b/Engine/Plugins/Developer/Concert/ConcertApp/MultiUserServer/Source/MultiUserServer/Private/Widgets/Clients/Logging/Util/MessageActionUtils.h
new file mode 100644
--- /dev/null
+++ b/Engine/Plugins/Developer/Concert/ConcertApp/MultiUserServer/Source/MultiUserServer/Private/Widgets/Clients/Logging/Util/MessageActionUtils.h
@@ -0,0 +1,17 @@
+// Copyright Epic Games, Inc. All Rights Reserved.
+
+#pragma once
+
+#include "CoreMinimal.h"
+#include "ConcertTransportEvents.h"
+
+namespace UE::MultiUserServer::MessageActionUtils
+{
+	/** Gets all message action names */
+	TSet<FName> GetAllMessageActionNames();
+	TSet<EConcertLogMessageAction> GetAllMessageActions();
+
+	FName ConvertActionToName(EConcertLogMessageAction MessageAction);
+	FString GetActionDisplayString(EConcertLogMessageAction MessageAction);
+	FString GetActionDisplayString(FName MessageActionName);
+};
diff --git a/Engine/Plugins/Developer/Concert/ConcertApp/MultiUserServer/Source/MultiUserServer/Private/Widgets/Clients/Logging/Util/MessageTypeUtils.cpp b/Engine/Plugins/Developer/Concert/ConcertApp/MultiUserServer/Source/MultiUserServer/Private/Widgets/Clients/Logging/Util/MessageTypeUtils.cpp
new file mode 100644
--- /dev/null
+++ b/Engine/Plugins/Developer/Concert/ConcertApp/MultiUserServer/Source/MultiUserServer/Private/Widgets/Clients/Logging/Util/MessageTypeUtils.cpp
@@ -0,0 +1,105 @@
+// Copyright Epic Games, Inc. All Rights Reserved.
+
+#include "MessageTypeUtils.h"
+
+#include "ConcertTransportMessages.h"
+#include "UObject/UObjectIterator.h"
+
+namespace UE::MultiUserServer::MessageTypeUtils
+{
+	static TSet<FName> FilterAllMessageTypes(TFunctionRef<bool(UStruct*)> ConsumerFunc);
+	
+	TSet<FName> GetAllMessageTypeNames()
+	{
+		static TSet<FName> CachedSubStructs = FilterAllMessageTypes([](UStruct* MessageType)
+		{
+			const TSet<UStruct*> Excluded = {
+				FConcertMessageData::StaticStruct(),
+				FConcertEventData::StaticStruct(),
+				FConcertRequestData::StaticStruct(),
+				FConcertResponseData::StaticStruct()
+			};
+
+			return !Excluded.Contains(MessageType);
+		});
+
+		return CachedSubStructs;
+	}
+
+	TSet<FName> GetAllMessageTypeNames_EventsOnly()
+	{
+		static TSet<FName> CachedSubStructs = FilterAllMessageTypes([](UStruct* MessageType)
+		{
+			return MessageType->IsChildOf(FConcertEventData::StaticStruct());
+		});
+
+		return CachedSubStructs;
+	}
+	
+	TSet<FName> GetAllMessageTypeNames_RequestsOnly()
+	{
+		static TSet<FName> CachedSubStructs = FilterAllMessageTypes([](UStruct* MessageType)
+		{
+			return MessageType->IsChildOf(FConcertRequestData::StaticStruct());
+		});
+
+		return CachedSubStructs;
+	}
+	
+	TSet<FName> GetAllMessageTypeNames_ResponseOnlyOnly()
+	{
+		static TSet<FName> CachedSubStructs = FilterAllMessageTypes([](UStruct* MessageType)
+		{
+			return MessageType->IsChildOf(FConcertResponseData::StaticStruct());
+		});
+
+		return CachedSubStructs;
+	}
+	
+	TSet<FName> GetAllMessageTypeNames_AcksOnlyOnly()
+	{
+		return { FConcertAckData::StaticStruct()->GetFName() };
+	}
+
+	FString SanitizeMessageTypeName(FName MessageTypeName)
+	{
+		static TMap<FName, FString> CachedSanitizedNames = []()
+		{
+			TMap<FName, FString> Result;
+			const TSet<UStruct*> Excluded = {
+				FConcertMessageData::StaticStruct(),
+				FConcertEventData::StaticStruct(),
+				FConcertRequestData::StaticStruct(),
+				FConcertResponseData::StaticStruct()
+			}; 
+			for (TObjectIterator<UStruct> ClassIt; ClassIt; ++ClassIt)
+			{
+				if (ClassIt->IsChildOf(FConcertMessageData::StaticStruct())
+					&& !Excluded.Contains(*ClassIt))
+				{
+					FString ClassName = ClassIt->GetName();
+					ClassName.RemoveFromStart("ConcertAdmin_");
+					ClassName.RemoveFromStart("ConcertSession_");
+					Result.Add(ClassIt->GetFName(), ClassName);
+				}
+			}
+			return Result;
+		}();
+
+		return CachedSanitizedNames[MessageTypeName];
+	}
+	
+	static TSet<FName> FilterAllMessageTypes(TFunctionRef<bool(UStruct*)> ConsumerFunc)
+	{
+		TSet<FName> Result;
+		for (TObjectIterator<UStruct> ClassIt; ClassIt; ++ClassIt)
+		{
+			if (ClassIt->IsChildOf(FConcertMessageData::StaticStruct())
+				&& ConsumerFunc(*ClassIt))
+			{
+				Result.Add(ClassIt->GetFName());
+			}
+		}
+		return Result;
+	}
+};
\ No newline at end of file
diff --git a/Engine/Plugins/Developer/Concert/ConcertApp/MultiUserServer/Source/MultiUserServer/Private/Widgets/Clients/Logging/Util/MessageTypeUtils.h b/Engine/Plugins/Developer/Concert/ConcertApp/MultiUserServer/Source/MultiUserServer/Private/Widgets/Clients/Logging/Util/MessageTypeUtils.h
new file mode 100644
--- /dev/null
+++ b/Engine/Plugins/Developer/Concert/ConcertApp/MultiUserServer/Source/MultiUserServer/Private/Widgets/Clients/Logging/Util/MessageTypeUtils.h
@@ -0,0 +1,19 @@
+// Copyright Epic Games, Inc. All Rights Reserved.
+
+#pragma once
+
+#include "CoreMinimal.h"
+
+namespace UE::MultiUserServer::MessageTypeUtils
+{
+	/** Gets all message type names (sub-structs of FConcertMessageData - see FConcertLog::MessageTypeName) */
+	TSet<FName> GetAllMessageTypeNames();
+	
+	TSet<FName> GetAllMessageTypeNames_EventsOnly();
+	TSet<FName> GetAllMessageTypeNames_RequestsOnly();
+	TSet<FName> GetAllMessageTypeNames_ResponseOnlyOnly();
+	TSet<FName> GetAllMessageTypeNames_AcksOnlyOnly();
+
+	/** Gets a readable version of a message type name */
+	FString SanitizeMessageTypeName(FName MessageTypeName);
+};
diff --git a/Engine/Plugins/Developer/Concert/ConcertApp/MultiUserServer/Source/MultiUserServer/Private/Widgets/Clients/SConcertClientsTabView.cpp b/Engine/Plugins/Developer/Concert/ConcertApp/MultiUserServer/Source/MultiUserServer/Private/Widgets/Clients/SConcertClientsTabView.cpp
new file mode 100644
--- /dev/null
+++ b/Engine/Plugins/Developer/Concert/ConcertApp/MultiUserServer/Source/MultiUserServer/Private/Widgets/Clients/SConcertClientsTabView.cpp
@@ -0,0 +1,122 @@
+// Copyright Epic Games, Inc. All Rights Reserved.
+
+#include "SConcertClientsTabView.h"
+
+#include "IConcertServer.h"
+#include "IConcertSyncServer.h"
+#include "Logging/Filter/ConcertLogFilter_FrontendRoot.h"
+#include "Logging/Source/GlobalLogSource.h"
+#include "Logging/Util/ConcertLogTokenizer.h"
+#include "SPromptConcertLoggingEnabled.h"
+#include "Widgets/Clients/Logging/SConcertTransportLog.h"
+
+#include "Framework/Docking/TabManager.h"
+#include "Widgets/Docking/SDockTab.h"
+
+#define LOCTEXT_NAMESPACE "UnrealMultiUserUI"
+
+const FName SConcertClientsTabView::ClientBrowserTabId("ClientBrowserTabId");
+const FName SConcertClientsTabView::GlobalLogTabId("GlobalLogTabId");
+
+SConcertClientsTabView::~SConcertClientsTabView()
+{
+	ConcertTransportEvents::OnConcertTransportLoggingEnabledChangedEvent().RemoveAll(this);
+}
+
+void SConcertClientsTabView::Construct(const FArguments& InArgs, FName InStatusBarID, TSharedRef<IConcertSyncServer> InServer)
+{
+	LogTokenizer = MakeShared<FConcertLogTokenizer>();
+	Server = MoveTemp(InServer);
+	
+	SConcertTabViewWithManagerBase::Construct(
+		SConcertTabViewWithManagerBase::FArguments()
+		.ConstructUnderWindow(InArgs._ConstructUnderWindow)
+		.ConstructUnderMajorTab(InArgs._ConstructUnderMajorTab)
+		.CreateTabs(FCreateTabs::CreateLambda([this, &InArgs](const TSharedRef<FTabManager>& InTabManager, const TSharedRef<FTabManager::FLayout>& InLayout)
+		{
+			CreateTabs(InTabManager, InLayout, InArgs);
+		}))
+		.OverlayTabs(this, &SConcertClientsTabView::SetupLoggingPromptOverlay)
+		.LayoutName("ConcertClientsTabView_v0.1"),
+		InStatusBarID
+	);
+}
+
+void SConcertClientsTabView::CreateTabs(const TSharedRef<FTabManager>& InTabManager, const TSharedRef<FTabManager::FLayout>& InLayout, const FArguments& InArgs)
+{
+	InTabManager->RegisterTabSpawner(ClientBrowserTabId, FOnSpawnTab::CreateSP(this, &SConcertClientsTabView::SpawnClientBrowserTab))
+		.SetDisplayName(LOCTEXT("ClientBrowserTabLabel", "Clients"));
+	InTabManager->RegisterTabSpawner(GlobalLogTabId, FOnSpawnTab::CreateSP(this, &SConcertClientsTabView::SpawnGlobalLogTab))
+		.SetDisplayName(LOCTEXT("GlobalLogTabLabel", "Global Log"));
+	InLayout->AddArea
+		(
+			FTabManager::NewPrimaryArea()
+				->SetOrientation(Orient_Vertical)
+				->Split
+				(
+					FTabManager::NewStack()
+					->SetSizeCoefficient(0.5f)
+					->AddTab(ClientBrowserTabId, ETabState::OpenedTab)
+				)
+				->Split
+				(
+					FTabManager::NewStack()
+					->SetSizeCoefficient(0.5f)
+					->AddTab(GlobalLogTabId, ETabState::OpenedTab)
+				)
+		);
+}
+
+TSharedRef<SDockTab> SConcertClientsTabView::SpawnClientBrowserTab(const FSpawnTabArgs& InTabArgs)
+{
+	return SNew(SDockTab)
+		.Label(LOCTEXT("ClientBrowserTabLabel", "Clients"))
+		.TabRole(PanelTab)
+		[
+			SNullWidget::NullWidget
+		]; 
+}
+
+TSharedRef<SDockTab> SConcertClientsTabView::SpawnGlobalLogTab(const FSpawnTabArgs& InTabArgs)
+{
+	// TODO: Load from config file
+	constexpr size_t LogCapacity = 500000;
+	
+	return SNew(SDockTab)
+		.Label(LOCTEXT("GlobalLogTabLabel", "Global Log"))
+		.TabRole(PanelTab)
+		[
+			SNew(SConcertTransportLog, MakeShared<FGlobalLogSource>(LogCapacity))
+			.Filter(UE::MultiUserServer::MakeGlobalLogFilter(LogTokenizer.ToSharedRef()))
+		]; 
+}
+
+TSharedRef<SWidget> SConcertClientsTabView::SetupLoggingPromptOverlay(const TSharedRef<SWidget>& TabsWidget)
+{
+	SOverlay::FOverlaySlot* MessageSlot;
+	EnableLoggingPromptOverlay =  SNew(SOverlay)
+		+SOverlay::Slot()
+		[
+			TabsWidget
+		]
+		+SOverlay::Slot()
+		.Expose(MessageSlot);
+
+	MessageSlot->AttachWidget(SAssignNew(EnableLoggingPrompt, SPromptConcertLoggingEnabled));
+	ConcertTransportEvents::OnConcertTransportLoggingEnabledChangedEvent().AddSP(this, &SConcertClientsTabView::OnConcertLoggingEnabledChanged);
+	return EnableLoggingPromptOverlay.ToSharedRef();
+}
+
+void SConcertClientsTabView::OnConcertLoggingEnabledChanged(bool bNewEnabled)
+{
+	if (bNewEnabled)
+	{
+		EnableLoggingPromptOverlay->RemoveSlot(EnableLoggingPrompt.ToSharedRef());
+	}
+	else
+	{
+		EnableLoggingPromptOverlay->AddSlot().AttachWidget(SAssignNew(EnableLoggingPrompt, SPromptConcertLoggingEnabled));
+	}
+}
+
+#undef LOCTEXT_NAMESPACE
diff --git a/Engine/Plugins/Developer/Concert/ConcertApp/MultiUserServer/Source/MultiUserServer/Private/Widgets/Clients/SConcertClientsTabView.h b/Engine/Plugins/Developer/Concert/ConcertApp/MultiUserServer/Source/MultiUserServer/Private/Widgets/Clients/SConcertClientsTabView.h
new file mode 100644
--- /dev/null
+++ b/Engine/Plugins/Developer/Concert/ConcertApp/MultiUserServer/Source/MultiUserServer/Private/Widgets/Clients/SConcertClientsTabView.h
@@ -0,0 +1,54 @@
+// Copyright Epic Games, Inc. All Rights Reserved.
+
+#pragma once
+
+#include "CoreMinimal.h"
+#include "Widgets/DeclarativeSyntaxSupport.h"
+#include "Widgets/SConcertTabViewWithManagerBase.h"
+
+struct FConcertClientInfo;
+class IConcertSyncServer;
+class FConcertLogTokenizer;
+class FTabManager;
+class SDockTab;
+class SPromptConcertLoggingEnabled;
+class SWidget;
+class SWindow;
+
+/** Manages the UI logic of the Clients tab */
+class SConcertClientsTabView : public SConcertTabViewWithManagerBase
+{
+public:
+	
+	static const FName ClientBrowserTabId;
+	static const FName GlobalLogTabId;
+
+	SLATE_BEGIN_ARGS(SConcertClientsTabView)
+	{}
+		SLATE_ARGUMENT(TSharedPtr<SDockTab>, ConstructUnderMajorTab)
+		SLATE_ARGUMENT(TSharedPtr<SWindow>, ConstructUnderWindow)
+	SLATE_END_ARGS()
+	virtual ~SConcertClientsTabView() override;
+
+	void Construct(const FArguments& InArgs, FName InStatusBarID, TSharedRef<IConcertSyncServer> InServer);
+
+private:
+
+	/** Used to overlay EnableLoggingPrompt over the tabs */
+	TSharedPtr<SOverlay> EnableLoggingPromptOverlay;
+	/** Reminds the user to enable logging */
+	TSharedPtr<SPromptConcertLoggingEnabled> EnableLoggingPrompt;
+
+	/** Used by various systems to convert logs to text */
+	TSharedPtr<FConcertLogTokenizer> LogTokenizer;
+
+	/** Used to look up client info */
+	TSharedPtr<IConcertSyncServer> Server;
+	
+	void CreateTabs(const TSharedRef<FTabManager>& InTabManager, const TSharedRef<FTabManager::FLayout>& InLayout, const FArguments& InArgs);
+	TSharedRef<SDockTab> SpawnClientBrowserTab(const FSpawnTabArgs& InTabArgs);
+	TSharedRef<SDockTab> SpawnGlobalLogTab(const FSpawnTabArgs& InTabArgs);
+
+	TSharedRef<SWidget> SetupLoggingPromptOverlay(const TSharedRef<SWidget>& TabsWidget);
+	void OnConcertLoggingEnabledChanged(bool bNewEnabled);
+};
diff --git a/Engine/Plugins/Developer/Concert/ConcertApp/MultiUserServer/Source/MultiUserServer/Private/Widgets/Clients/SPromptConcertLoggingEnabled.cpp b/Engine/Plugins/Developer/Concert/ConcertApp/MultiUserServer/Source/MultiUserServer/Private/Widgets/Clients/SPromptConcertLoggingEnabled.cpp
new file mode 100644
--- /dev/null
+++ b/Engine/Plugins/Developer/Concert/ConcertApp/MultiUserServer/Source/MultiUserServer/Private/Widgets/Clients/SPromptConcertLoggingEnabled.cpp
@@ -0,0 +1,59 @@
+// Copyright Epic Games, Inc. All Rights Reserved.
+
+#include "SPromptConcertLoggingEnabled.h"
+
+#include "ConcertTransportEvents.h"
+#include "Styling/AppStyle.h"
+#include "Widgets/Input/SButton.h"
+#include "Widgets/Layout/SBorder.h"
+#include "Widgets/Layout/SBox.h"
+#include "Widgets/SBoxPanel.h"
+#include "Widgets/Text/STextBlock.h"
+
+#define LOCTEXT_NAMESPACE "UnrealMultiUserUI"
+
+void SPromptConcertLoggingEnabled::Construct(const FArguments& InArgs)
+{
+	ChildSlot
+	[
+		SNew(SBorder)
+		.BorderImage(FAppStyle::Get().GetBrush("ToolPanel.GroupBorder"))
+		.BorderBackgroundColor(FSlateColor(FLinearColor(0.6, 0.6, 0.6, 0.8f)))
+		.Padding(0)
+		[
+			SNew(SBox)
+			.Padding(0.f, 0.f, 0.f, 50.f)
+			.HAlign(HAlign_Center)
+			.VAlign(VAlign_Center)
+			.Clipping(EWidgetClipping::ClipToBounds)
+			[
+				SNew(SVerticalBox)
+
+				+SVerticalBox::Slot()
+				.AutoHeight()
+				.HAlign(HAlign_Center)
+				[
+					SNew(STextBlock)
+					.Text(LOCTEXT("EnableLoggingVisibility", "Execute Concert.EnableLogging to use this view."))
+					.Justification(ETextJustify::Center)
+				]
+
+				+SVerticalBox::Slot()
+				.Padding(0, 4, 0, 0)
+				.AutoHeight()
+				.HAlign(HAlign_Center)
+				[
+					SNew(SButton)
+					.OnClicked_Lambda([](){ ConcertTransportEvents::SetLoggingEnabled(true); return FReply::Handled(); })
+					.ToolTipText(LOCTEXT("EnableLoggingVisibility.Button.Tooltip", "Executes Concert.EnableLogging"))
+					[
+						SNew(STextBlock)
+							.Text(LOCTEXT("EnableLoggingVisibility.Button.Text", "Enable logging"))
+					]
+				]
+			]
+		]
+	];
+}
+
+#undef LOCTEXT_NAMESPACE
\ No newline at end of file
diff --git a/Engine/Plugins/Developer/Concert/ConcertApp/MultiUserServer/Source/MultiUserServer/Private/Widgets/Clients/SPromptConcertLoggingEnabled.h b/Engine/Plugins/Developer/Concert/ConcertApp/MultiUserServer/Source/MultiUserServer/Private/Widgets/Clients/SPromptConcertLoggingEnabled.h
new file mode 100644
--- /dev/null
+++ b/Engine/Plugins/Developer/Concert/ConcertApp/MultiUserServer/Source/MultiUserServer/Private/Widgets/Clients/SPromptConcertLoggingEnabled.h
@@ -0,0 +1,18 @@
+// Copyright Epic Games, Inc. All Rights Reserved.
+
+#pragma once
+
+#include "CoreMinimal.h"
+#include "Widgets/DeclarativeSyntaxSupport.h"
+#include "Widgets/SCompoundWidget.h"
+
+/** Placed in an overlay, reminds the user to execute Concert.EnableLogging if logging is not enabled. */
+class SPromptConcertLoggingEnabled : public SCompoundWidget
+{
+public:
+
+	SLATE_BEGIN_ARGS(SPromptConcertLoggingEnabled) {}
+	SLATE_END_ARGS()
+
+	void Construct(const FArguments& InArgs);
+};
diff --git a/Engine/Plugins/Developer/Concert/ConcertApp/MultiUserServer/Source/MultiUserServer/Private/Widgets/SConcertTabViewWithManagerBase.cpp b/Engine/Plugins/Developer/Concert/ConcertApp/MultiUserServer/Source/MultiUserServer/Private/Widgets/SConcertTabViewWithManagerBase.cpp
--- a/Engine/Plugins/Developer/Concert/ConcertApp/MultiUserServer/Source/MultiUserServer/Private/Widgets/SConcertTabViewWithManagerBase.cpp
+++ b/Engine/Plugins/Developer/Concert/ConcertApp/MultiUserServer/Source/MultiUserServer/Private/Widgets/SConcertTabViewWithManagerBase.cpp
@@ -11,7 +11,7 @@ void SConcertTabViewWithManagerBase::Construct(const FArguments& InArgs, FName I
 		SConcertTabViewBase::FArguments()
 		.Content()
 		[
-			CreateTabs(InArgs)
+			InArgs._OverlayTabs.IsBound() ? InArgs._OverlayTabs.Execute(CreateTabs(InArgs)) : CreateTabs(InArgs)
 		],
 		InStatusBarId
 		);
diff --git a/Engine/Plugins/Developer/Concert/ConcertApp/MultiUserServer/Source/MultiUserServer/Private/Widgets/SConcertTabViewWithManagerBase.h b/Engine/Plugins/Developer/Concert/ConcertApp/MultiUserServer/Source/MultiUserServer/Private/Widgets/SConcertTabViewWithManagerBase.h
--- a/Engine/Plugins/Developer/Concert/ConcertApp/MultiUserServer/Source/MultiUserServer/Private/Widgets/SConcertTabViewWithManagerBase.h
+++ b/Engine/Plugins/Developer/Concert/ConcertApp/MultiUserServer/Source/MultiUserServer/Private/Widgets/SConcertTabViewWithManagerBase.h
@@ -13,6 +13,7 @@ class SConcertTabViewWithManagerBase : public SConcertTabViewBase
 public:
 
 	DECLARE_DELEGATE_TwoParams(FCreateTabs, const TSharedRef<FTabManager>& /*TabManager*/, const TSharedRef<FTabManager::FLayout>& /*Layout*/);
+	DECLARE_DELEGATE_RetVal_OneParam(TSharedRef<SWidget>, FOverlayTabs, const TSharedRef<SWidget>& /* Tabs */);
 	
 	SLATE_BEGIN_ARGS(SConcertTabViewWithManagerBase) {}
 		/** Which major tab to construct the sub-tabs under. */
@@ -21,6 +22,8 @@ public:
 		SLATE_ARGUMENT(TSharedPtr<SWindow>, ConstructUnderWindow)
 		/** Callback for creating the sub-tabs */
 		SLATE_EVENT(FCreateTabs, CreateTabs)
+		/** Optional function to overlay all the tabs (for example to overlay a message to execute a console command first) */
+		SLATE_EVENT(FOverlayTabs, OverlayTabs)
 		/** Name to give the layout. Important for saving config. */
 		SLATE_ARGUMENT(FName, LayoutName)
 	SLATE_END_ARGS()
diff --git a/Engine/Plugins/Developer/Concert/ConcertApp/MultiUserServer/Source/MultiUserServer/Private/Widgets/SessionTabs/Archived/ArchivedSessionHistoryController.cpp b/Engine/Plugins/Developer/Concert/ConcertApp/MultiUserServer/Source/MultiUserServer/Private/Widgets/SessionTabs/Archived/ArchivedSessionHistoryController.cpp
--- a/Engine/Plugins/Developer/Concert/ConcertApp/MultiUserServer/Source/MultiUserServer/Private/Widgets/SessionTabs/Archived/ArchivedSessionHistoryController.cpp
+++ b/Engine/Plugins/Developer/Concert/ConcertApp/MultiUserServer/Source/MultiUserServer/Private/Widgets/SessionTabs/Archived/ArchivedSessionHistoryController.cpp
@@ -2,7 +2,7 @@
 
 #include "ArchivedSessionHistoryController.h"
 
-#include "MultiUserServerUserSettings.h"
+#include "Settings/MultiUserServerColumnVisibilitySettings.h"
 
 FArchivedSessionHistoryController::FArchivedSessionHistoryController(FGuid SessionId, TSharedRef<IConcertSyncServer> SyncServer, SSessionHistory::FArguments Arguments)
 	: FServerSessionHistoryControllerBase(MoveTemp(SessionId), MoveTemp(Arguments))
@@ -23,10 +23,10 @@ namespace UE::MultiUserServer
 		static SSessionHistory::FArguments MakeArgumentsForInspector(SSessionHistory::FArguments&& InArgs)
 		{
 			return InArgs
-				.ColumnVisibilitySnapshot(UMultiUserServerUserSettings::GetUserSettings()->GetArchivedActivityBrowserColumnVisibility())
+				.ColumnVisibilitySnapshot(UMultiUserServerColumnVisibilitySettings::GetSettings()->GetArchivedActivityBrowserColumnVisibility())
 				.SaveColumnVisibilitySnapshot_Lambda([](const FColumnVisibilitySnapshot& Snapshot)
 				{
-					UMultiUserServerUserSettings::GetUserSettings()->SetArchivedActivityBrowserColumnVisibility(Snapshot);
+					UMultiUserServerColumnVisibilitySettings::GetSettings()->SetArchivedActivityBrowserColumnVisibility(Snapshot);
 				});
 		}
 
@@ -36,12 +36,12 @@ namespace UE::MultiUserServer
 			FInspectorSessionHistoryController(FGuid SessionId, TSharedRef<IConcertSyncServer> SyncServer, SSessionHistory::FArguments Arguments)
 				: FArchivedSessionHistoryController(MoveTemp(SessionId), MoveTemp(SyncServer), MakeArgumentsForInspector(MoveTemp(Arguments)))
 			{
-				UMultiUserServerUserSettings::GetUserSettings()->OnArchivedActivityBrowserColumnVisibility().AddRaw(this, &FInspectorSessionHistoryController::OnActivityListColumnVisibilitySettingsUpdated);
+				UMultiUserServerColumnVisibilitySettings::GetSettings()->OnArchivedActivityBrowserColumnVisibility().AddRaw(this, &FInspectorSessionHistoryController::OnActivityListColumnVisibilitySettingsUpdated);
 			}
 		
 			virtual ~FInspectorSessionHistoryController() override
 			{
-				if (UMultiUserServerUserSettings* Settings = UMultiUserServerUserSettings::GetUserSettings(); IsValid(Settings))
+				if (UMultiUserServerColumnVisibilitySettings* Settings = UMultiUserServerColumnVisibilitySettings::GetSettings(); IsValid(Settings))
 				{
 					Settings->OnArchivedActivityBrowserColumnVisibility().RemoveAll(this);
 				}
@@ -56,10 +56,10 @@ namespace UE::MultiUserServer
 		static SSessionHistory::FArguments MakeArgumentsForDeleteDialog(SSessionHistory::FArguments&& InArgs)
 		{
 			return InArgs
-				.ColumnVisibilitySnapshot(UMultiUserServerUserSettings::GetUserSettings()->GetDeleteActivityDialogColumnVisibility())
+				.ColumnVisibilitySnapshot(UMultiUserServerColumnVisibilitySettings::GetSettings()->GetDeleteActivityDialogColumnVisibility())
 				.SaveColumnVisibilitySnapshot_Lambda([](const FColumnVisibilitySnapshot& Snapshot)
 				{
-					UMultiUserServerUserSettings::GetUserSettings()->SetDeleteActivityDialogColumnVisibility(Snapshot);
+					UMultiUserServerColumnVisibilitySettings::GetSettings()->SetDeleteActivityDialogColumnVisibility(Snapshot);
 				});
 		}
 
@@ -70,12 +70,12 @@ namespace UE::MultiUserServer
 			FDeleteDialogSessionHistoryController(FGuid SessionId, TSharedRef<IConcertSyncServer> SyncServer, SSessionHistory::FArguments Arguments)
 				: FArchivedSessionHistoryController(MoveTemp(SessionId), MoveTemp(SyncServer), MakeArgumentsForDeleteDialog(MoveTemp(Arguments)))
 			{
-				UMultiUserServerUserSettings::GetUserSettings()->OnDeleteActivityDialogColumnVisibility().AddRaw(this, &FDeleteDialogSessionHistoryController::OnActivityListColumnVisibilitySettingsUpdated);
+				UMultiUserServerColumnVisibilitySettings::GetSettings()->OnDeleteActivityDialogColumnVisibility().AddRaw(this, &FDeleteDialogSessionHistoryController::OnActivityListColumnVisibilitySettingsUpdated);
 			}
 		
 			virtual ~FDeleteDialogSessionHistoryController() override
 			{
-				if (UMultiUserServerUserSettings* Settings = UMultiUserServerUserSettings::GetUserSettings(); IsValid(Settings))
+				if (UMultiUserServerColumnVisibilitySettings* Settings = UMultiUserServerColumnVisibilitySettings::GetSettings(); IsValid(Settings))
 				{
 					Settings->OnDeleteActivityDialogColumnVisibility().RemoveAll(this);
 				}
diff --git a/Engine/Plugins/Developer/Concert/ConcertApp/MultiUserServer/Source/MultiUserServer/Private/Widgets/SessionTabs/Live/LiveServerSessionHistoryController.cpp b/Engine/Plugins/Developer/Concert/ConcertApp/MultiUserServer/Source/MultiUserServer/Private/Widgets/SessionTabs/Live/LiveServerSessionHistoryController.cpp
--- a/Engine/Plugins/Developer/Concert/ConcertApp/MultiUserServer/Source/MultiUserServer/Private/Widgets/SessionTabs/Live/LiveServerSessionHistoryController.cpp
+++ b/Engine/Plugins/Developer/Concert/ConcertApp/MultiUserServer/Source/MultiUserServer/Private/Widgets/SessionTabs/Live/LiveServerSessionHistoryController.cpp
@@ -5,17 +5,17 @@
 #include "ConcertSyncSessionDatabase.h"
 #include "IConcertSession.h"
 #include "IConcertSyncServer.h"
-#include "MultiUserServerUserSettings.h"
+#include "Settings/MultiUserServerColumnVisibilitySettings.h"
 
 namespace UE::MultiUserServer::Private
 {
 	static SSessionHistory::FArguments MakeLiveTabSessionHistoryArguments()
 	{
 		return SSessionHistory::FArguments()
-			.ColumnVisibilitySnapshot(UMultiUserServerUserSettings::GetUserSettings()->GetLiveActivityBrowserColumnVisibility())
+			.ColumnVisibilitySnapshot(UMultiUserServerColumnVisibilitySettings::GetSettings()->GetLiveActivityBrowserColumnVisibility())
 			.SaveColumnVisibilitySnapshot_Lambda([](const FColumnVisibilitySnapshot& Snapshot)
 			{
-				UMultiUserServerUserSettings::GetUserSettings()->SetLiveActivityBrowserColumnVisibility(Snapshot);
+				UMultiUserServerColumnVisibilitySettings::GetSettings()->SetLiveActivityBrowserColumnVisibility(Snapshot);
 			});
 	}
 }
@@ -25,7 +25,7 @@ FLiveServerSessionHistoryController::FLiveServerSessionHistoryController(TShared
 	, SyncServer(MoveTemp(SyncServer))
 {
 	ReloadActivities();
-	UMultiUserServerUserSettings::GetUserSettings()->OnLiveActivityBrowserColumnVisibility().AddRaw(this, &FLiveServerSessionHistoryController::OnActivityListColumnVisibilitySettingsUpdated);
+	UMultiUserServerColumnVisibilitySettings::GetSettings()->OnLiveActivityBrowserColumnVisibility().AddRaw(this, &FLiveServerSessionHistoryController::OnActivityListColumnVisibilitySettingsUpdated);
 
 	if (const TOptional<FConcertSyncSessionDatabaseNonNullPtr> Database = SyncServer->GetLiveSessionDatabase(GetSessionId()))
 	{
@@ -35,7 +35,7 @@ FLiveServerSessionHistoryController::FLiveServerSessionHistoryController(TShared
 
 FLiveServerSessionHistoryController::~FLiveServerSessionHistoryController()
 {
-	if (UMultiUserServerUserSettings* Settings = UMultiUserServerUserSettings::GetUserSettings(); IsValid(Settings))
+	if (UMultiUserServerColumnVisibilitySettings* Settings = UMultiUserServerColumnVisibilitySettings::GetSettings(); IsValid(Settings))
 	{
 		Settings->OnLiveActivityBrowserColumnVisibility().RemoveAll(this);
 	}
diff --git a/Engine/Plugins/Developer/Concert/ConcertApp/MultiUserServer/Source/MultiUserServer/Private/Widgets/SessionTabs/Live/PackageViewer/ConcertSessionPackageViewerController.cpp b/Engine/Plugins/Developer/Concert/ConcertApp/MultiUserServer/Source/MultiUserServer/Private/Widgets/SessionTabs/Live/PackageViewer/ConcertSessionPackageViewerController.cpp
--- a/Engine/Plugins/Developer/Concert/ConcertApp/MultiUserServer/Source/MultiUserServer/Private/Widgets/SessionTabs/Live/PackageViewer/ConcertSessionPackageViewerController.cpp
+++ b/Engine/Plugins/Developer/Concert/ConcertApp/MultiUserServer/Source/MultiUserServer/Private/Widgets/SessionTabs/Live/PackageViewer/ConcertSessionPackageViewerController.cpp
@@ -4,9 +4,9 @@
 
 #include "ConcertSyncSessionDatabase.h"
 #include "IConcertSyncServer.h"
-#include "MultiUserServerUserSettings.h"
+#include "IConcertSession.h"
+#include "Settings/MultiUserServerColumnVisibilitySettings.h"
 #include "SConcertSessionPackageViewer.h"
-#include "Concert/Private/ConcertServerSession.h"
 
 FConcertSessionPackageViewerController::FConcertSessionPackageViewerController(TSharedRef<IConcertServerSession> InspectedSession, TSharedRef<IConcertSyncServer> SyncServer)
 	: InspectedSession(MoveTemp(InspectedSession))
@@ -14,7 +14,7 @@ FConcertSessionPackageViewerController::FConcertSessionPackageViewerController(T
 	, PackageViewer(MakePackageViewer())
 {
 	ReloadActivities();
-	UMultiUserServerUserSettings::GetUserSettings()->OnLiveSessionContentColumnVisibility().AddRaw(this, &FConcertSessionPackageViewerController::OnSessionContentColumnVisibilitySettingsUpdated);
+	UMultiUserServerColumnVisibilitySettings::GetSettings()->OnLiveSessionContentColumnVisibility().AddRaw(this, &FConcertSessionPackageViewerController::OnSessionContentColumnVisibilitySettingsUpdated);
 	
 	if (const TOptional<FConcertSyncSessionDatabaseNonNullPtr> Database = SyncServer->GetLiveSessionDatabase(InspectedSession->GetId()))
 	{
@@ -24,7 +24,7 @@ FConcertSessionPackageViewerController::FConcertSessionPackageViewerController(T
 
 FConcertSessionPackageViewerController::~FConcertSessionPackageViewerController()
 {
-	if (UMultiUserServerUserSettings* Settings = UMultiUserServerUserSettings::GetUserSettings(); IsValid(Settings))
+	if (UMultiUserServerColumnVisibilitySettings* Settings = UMultiUserServerColumnVisibilitySettings::GetSettings(); IsValid(Settings))
 	{
 		Settings->OnLiveSessionContentColumnVisibility().RemoveAll(this);
 	}
diff --git a/Engine/Plugins/Developer/Concert/ConcertApp/MultiUserServer/Source/MultiUserServer/Private/Widgets/SessionTabs/Live/PackageViewer/SConcertSessionPackageViewer.cpp b/Engine/Plugins/Developer/Concert/ConcertApp/MultiUserServer/Source/MultiUserServer/Private/Widgets/SessionTabs/Live/PackageViewer/SConcertSessionPackageViewer.cpp
--- a/Engine/Plugins/Developer/Concert/ConcertApp/MultiUserServer/Source/MultiUserServer/Private/Widgets/SessionTabs/Live/PackageViewer/SConcertSessionPackageViewer.cpp
+++ b/Engine/Plugins/Developer/Concert/ConcertApp/MultiUserServer/Source/MultiUserServer/Private/Widgets/SessionTabs/Live/PackageViewer/SConcertSessionPackageViewer.cpp
@@ -2,13 +2,12 @@
 
 #include "SConcertSessionPackageViewer.h"
 
-#include "MultiUserServerUserSettings.h"
 #include "PackageViewerColumns.h"
+#include "Settings/MultiUserServerColumnVisibilitySettings.h"
 
-#include "Session/Activity/PredefinedActivityColumns.h"
+#include "Framework/MultiBox/MultiBoxBuilder.h"
 #include "Session/Activity/SConcertSessionActivities.h"
 #include "Widgets/Input/SSearchBox.h"
-#include "Framework/MultiBox/MultiBoxBuilder.h"
 
 #define LOCTEXT_NAMESPACE "UnrealMultiUserUI"
 
@@ -41,10 +40,10 @@ void SConcertSessionPackageViewer::Construct(const FArguments& InArgs)
 		.TransactionActivitiesVisibility(EVisibility::Collapsed)
 		.DetailsAreaVisibility(EVisibility::Collapsed)
 		.IsAutoScrollEnabled(true)
-		.ColumnVisibilitySnapshot(UMultiUserServerUserSettings::GetUserSettings()->GetLiveSessionContentColumnVisibility())
+		.ColumnVisibilitySnapshot(UMultiUserServerColumnVisibilitySettings::GetSettings()->GetLiveSessionContentColumnVisibility())
 		.SaveColumnVisibilitySnapshot_Lambda([](const FColumnVisibilitySnapshot& Snapshot)
 		{
-			UMultiUserServerUserSettings::GetUserSettings()->SetLiveSessionContentColumnVisibility(Snapshot);
+			UMultiUserServerColumnVisibilitySettings::GetSettings()->SetLiveSessionContentColumnVisibility(Snapshot);
 		});
 
 	ChildSlot
diff --git a/Engine/Plugins/Developer/Concert/ConcertApp/MultiUserServer/Source/MultiUserServer/Private/Widgets/SessionTabs/Live/SConcertLiveSessionTabView.cpp b/Engine/Plugins/Developer/Concert/ConcertApp/MultiUserServer/Source/MultiUserServer/Private/Widgets/SessionTabs/Live/SConcertLiveSessionTabView.cpp
--- a/Engine/Plugins/Developer/Concert/ConcertApp/MultiUserServer/Source/MultiUserServer/Private/Widgets/SessionTabs/Live/SConcertLiveSessionTabView.cpp
+++ b/Engine/Plugins/Developer/Concert/ConcertApp/MultiUserServer/Source/MultiUserServer/Private/Widgets/SessionTabs/Live/SConcertLiveSessionTabView.cpp
@@ -5,8 +5,6 @@
 #include "Framework/Docking/TabManager.h"
 #include "Session/History/SSessionHistory.h"
 #include "Widgets/Docking/SDockTab.h"
-#include "Widgets/Layout/SBorder.h"
-#include "Widgets/SBoxPanel.h"
 
 #define LOCTEXT_NAMESPACE "SConcertSessionInspector"
 
@@ -24,7 +22,7 @@ void SConcertLiveSessionTabView::Construct(const FArguments& InArgs, const FRequ
 		{
 			CreateTabs(InTabManager, InLayout, InRequiredArgs);
 		}))
-		.LayoutName("ConcertSessionInspector_v0.1"),
+		.LayoutName("ConcertSessionInspector_v0.2"),
 		StatusBarId
 		);
 }
@@ -36,9 +34,6 @@ TSharedRef<SWidget> SConcertLiveSessionTabView::CreateTabs(const TSharedRef<FTab
 	
 	InTabManager->RegisterTabSpawner(SessionContentTabId, FOnSpawnTab::CreateSP(this, &SConcertLiveSessionTabView::SpawnSessionContent, InRequiredArgs.PackageViewer))
 		.SetDisplayName(LOCTEXT("SessionContentLabel", "Session Content"));
-	
-	InTabManager->RegisterTabSpawner(ConnectionMonitorTabId, FOnSpawnTab::CreateSP(this, &SConcertLiveSessionTabView::SpawnConnectionMonitor))
-		.SetDisplayName(LOCTEXT("ConnectionMonitorLabel", "Connection Monitor"));
 
 	InLayout->AddArea
 	(
@@ -93,14 +88,4 @@ TSharedRef<SDockTab> SConcertLiveSessionTabView::SpawnSessionContent(const FSpaw
 		];
 }
 
-TSharedRef<SDockTab> SConcertLiveSessionTabView::SpawnConnectionMonitor(const FSpawnTabArgs& Args)
-{
-	return SNew(SDockTab)
-		.Label(LOCTEXT("ConnectionMonitorLabel", "Connection Monitor"))
-		.TabRole(PanelTab)
-		[
-			SNullWidget::NullWidget
-		];
-}
-
 #undef LOCTEXT_NAMESPACE
diff --git a/Engine/Plugins/Developer/Concert/ConcertApp/MultiUserServer/Source/MultiUserServer/Private/Widgets/SessionTabs/Live/SConcertLiveSessionTabView.h b/Engine/Plugins/Developer/Concert/ConcertApp/MultiUserServer/Source/MultiUserServer/Private/Widgets/SessionTabs/Live/SConcertLiveSessionTabView.h
--- a/Engine/Plugins/Developer/Concert/ConcertApp/MultiUserServer/Source/MultiUserServer/Private/Widgets/SessionTabs/Live/SConcertLiveSessionTabView.h
+++ b/Engine/Plugins/Developer/Concert/ConcertApp/MultiUserServer/Source/MultiUserServer/Private/Widgets/SessionTabs/Live/SConcertLiveSessionTabView.h
@@ -58,5 +58,4 @@ private:
 	TSharedRef<SWidget> CreateTabs(const TSharedRef<FTabManager>& InTabManager, const TSharedRef<FTabManager::FLayout>& InLayout, const FRequiredWidgets& InRequiredArgs);
 	TSharedRef<SDockTab> SpawnActivityHistory(const FSpawnTabArgs& Args, TSharedRef<SSessionHistory> SessionHistory);
 	TSharedRef<SDockTab> SpawnSessionContent(const FSpawnTabArgs& Args, TSharedRef<SConcertSessionPackageViewer> PackageViewer);
-	TSharedRef<SDockTab> SpawnConnectionMonitor(const FSpawnTabArgs& Args);
 };
diff --git a/Engine/Plugins/Developer/Concert/ConcertApp/MultiUserServer/Source/MultiUserServer/Private/Window/ConcertServerWindowController.cpp b/Engine/Plugins/Developer/Concert/ConcertApp/MultiUserServer/Source/MultiUserServer/Private/Window/ConcertServerWindowController.cpp
--- a/Engine/Plugins/Developer/Concert/ConcertApp/MultiUserServer/Source/MultiUserServer/Private/Window/ConcertServerWindowController.cpp
+++ b/Engine/Plugins/Developer/Concert/ConcertApp/MultiUserServer/Source/MultiUserServer/Private/Window/ConcertServerWindowController.cpp
@@ -6,6 +6,7 @@
 #include "IConcertServer.h"
 #include "IConcertSyncServer.h"
 #include "Widgets/Browser/ConcertServerSessionBrowserController.h"
+#include "Widgets/Clients/ConcertClientsTabController.h"
 #include "Widgets/SessionTabs/Archived/ArchivedConcertSessionTab.h"
 #include "Widgets/SessionTabs/Live/LiveConcertSessionTab.h"
 
@@ -26,9 +27,11 @@ FConcertServerWindowController::FConcertServerWindowController(const FConcertSer
 	: MultiUserServerLayoutIni(Params.MultiUserServerLayoutIni)
 	, ServerInstance(Params.Server)
 	, SessionBrowserController(MakeShared<FConcertServerSessionBrowserController>())
+	, ClientsController(MakeShared<FConcertClientsTabController>())
 	, ConcertComponents(Params.AdditionalConcertComponents)
 {
 	ConcertComponents.Add(SessionBrowserController);
+	ConcertComponents.Add(ClientsController);
 }
 
 FConcertServerWindowController::~FConcertServerWindowController()
@@ -104,7 +107,7 @@ TSharedPtr<FConcertSessionTabBase> FConcertServerWindowController::GetOrRegister
 		return *FoundId;
 	}
 	
-	if (TSharedPtr<IConcertServerSession> Session = ServerInstance->GetConcertServer()->GetLiveSession(SessionId))
+	if (const TSharedPtr<IConcertServerSession> Session = ServerInstance->GetConcertServer()->GetLiveSession(SessionId))
 	{
 		const TSharedRef<FLiveConcertSessionTab> SessionTab = MakeShared<FLiveConcertSessionTab>(Session.ToSharedRef(), ServerInstance, RootWindow.ToSharedRef());
 		RegisteredSessions.Add(SessionId, SessionTab);
diff --git a/Engine/Plugins/Developer/Concert/ConcertApp/MultiUserServer/Source/MultiUserServer/Private/Window/ConcertServerWindowController.h b/Engine/Plugins/Developer/Concert/ConcertApp/MultiUserServer/Source/MultiUserServer/Private/Window/ConcertServerWindowController.h
--- a/Engine/Plugins/Developer/Concert/ConcertApp/MultiUserServer/Source/MultiUserServer/Private/Window/ConcertServerWindowController.h
+++ b/Engine/Plugins/Developer/Concert/ConcertApp/MultiUserServer/Source/MultiUserServer/Private/Window/ConcertServerWindowController.h
@@ -12,6 +12,7 @@ class FLiveConcertSessionTab;
 class IConcertServerSession;
 class IConcertSyncServer;
 class FConcertServerSessionBrowserController;
+class FConcertClientsTabController;
 class FOutputLogController;
 class SWindow;
 
@@ -39,7 +40,9 @@ public:
 	
 	FConcertServerWindowController(const FConcertServerWindowInitParams& Params);
 	~FConcertServerWindowController();
+	
 	TSharedRef<SWindow> CreateWindow();
+	TSharedPtr<SWindow> GetRootWindow() const { return RootWindow; }
 
 	/** Opens or draws attention to the tab for the given live or archived session ID */
 	void OpenSessionTab(const FGuid& SessionId);
@@ -62,6 +65,8 @@ private:
 	
 	/** Manages the session browser */
 	TSharedRef<FConcertServerSessionBrowserController> SessionBrowserController;
+	/** Manages the view of connected clients */
+	TSharedRef<FConcertClientsTabController> ClientsController;
 	/** Additional external concert components */
 	TArray<TSharedRef<IConcertComponent>> ConcertComponents;
 	
