Fix unity compile error

#rnx

diff --git a/Engine/Plugins/Developer/Concert/ConcertApp/MultiUserServer/Source/MultiUserServer/Private/Widgets/Clients/Logging/Filter/ConcertLogFilter_FrontendRoot.h b/Engine/Plugins/Developer/Concert/ConcertApp/MultiUserServer/Source/MultiUserServer/Private/Widgets/Clients/Logging/Filter/ConcertLogFilter_FrontendRoot.h
--- a/Engine/Plugins/Developer/Concert/ConcertApp/MultiUserServer/Source/MultiUserServer/Private/Widgets/Clients/Logging/Filter/ConcertLogFilter_FrontendRoot.h
+++ b/Engine/Plugins/Developer/Concert/ConcertApp/MultiUserServer/Source/MultiUserServer/Private/Widgets/Clients/Logging/Filter/ConcertLogFilter_FrontendRoot.h
@@ -8,6 +8,7 @@
 class FConcertLogTokenizer;
 class FConcertFrontendLogFilter;
 class FConcertFrontendLogFilter_TextSearch;
+class SWidget;
 
 /** A filter that contains multiple UI filters */
 class FConcertLogFilter_FrontendRoot : public FConcertLogFilter, public TSharedFromThis<FConcertLogFilter_FrontendRoot>
