Add client browser similar to content browser

diff --git a/Engine/Plugins/Developer/Concert/ConcertApp/MultiUserServer/Source/MultiUserServer/MultiUserServer.build.cs b/Engine/Plugins/Developer/Concert/ConcertApp/MultiUserServer/Source/MultiUserServer/MultiUserServer.build.cs
--- a/Engine/Plugins/Developer/Concert/ConcertApp/MultiUserServer/Source/MultiUserServer/MultiUserServer.build.cs
+++ b/Engine/Plugins/Developer/Concert/ConcertApp/MultiUserServer/Source/MultiUserServer/MultiUserServer.build.cs
@@ -19,6 +19,7 @@ namespace UnrealBuildTool.Rules
 					"ConcertSyncCore",
 					"ConcertTransport",
 					"InputCore",
+					"Messaging",
 					"OutputLog",
 					"Projects",
 					"Slate",
diff --git a/Engine/Plugins/Developer/Concert/ConcertApp/MultiUserServer/Source/MultiUserServer/Private/ConcertServerStyle.cpp b/Engine/Plugins/Developer/Concert/ConcertApp/MultiUserServer/Source/MultiUserServer/Private/ConcertServerStyle.cpp
--- a/Engine/Plugins/Developer/Concert/ConcertApp/MultiUserServer/Source/MultiUserServer/Private/ConcertServerStyle.cpp
+++ b/Engine/Plugins/Developer/Concert/ConcertApp/MultiUserServer/Source/MultiUserServer/Private/ConcertServerStyle.cpp
@@ -4,7 +4,9 @@
 
 #include "Framework/Application/SlateApplication.h"
 #include "Interfaces/IPluginManager.h"
+#include "Styling/SlateStyleMacros.h"
 #include "Styling/SlateStyleRegistry.h"
+#include "Styling/StyleColors.h"
 
 TSharedPtr<FSlateStyleSet> FConcertServerStyle::StyleInstance = nullptr;
 
@@ -40,7 +42,7 @@ FName FConcertServerStyle::GetStyleSetName()
 }
 
 #define IMAGE_PLUGIN_BRUSH( RelativePath, ... ) FSlateImageBrush( FConcertServerStyle::InContent( RelativePath, ".png" ), __VA_ARGS__ )
-#define IMAGE_BRUSH(RelativePath, ...) FSlateImageBrush(StyleSet->RootToContentDir(RelativePath, TEXT(".png")), __VA_ARGS__)
+#define IMAGE_PLUGIN_BOX_BRUSH( RelativePath, ... ) FSlateBoxBrush( FConcertServerStyle::InContent( RelativePath, ".png" ), __VA_ARGS__ )
 #define IMAGE_PLUGIN_BRUSH_SVG( RelativePath, ... ) FSlateVectorImageBrush( FConcertServerStyle::InContent(RelativePath, ".svg"), __VA_ARGS__)
 
 FString FConcertServerStyle::InContent(const FString& RelativePath, const ANSICHAR* Extension)
@@ -59,7 +61,8 @@ TSharedRef<FSlateStyleSet> FConcertServerStyle::Create()
 	const FVector2D Icon12x12(12.0f, 12.0f); 
 	const FVector2D Icon16x16(16.0f, 16.0f); 
 	const FVector2D Icon20x20(20.0f, 20.0f); 
-	const FVector2D Icon32x32(20.0f, 20.0f); 
+	const FVector2D Icon32x32(32.0f, 32.0f); 
+	const FVector2D Icon64x64(64.0f, 64.0f); 
 	
 	StyleSet->Set("Concert.MultiUser", new IMAGE_PLUGIN_BRUSH("Icons/icon_MultiUser_32x", Icon32x32));
 	StyleSet->Set("Concert.SessionContent.ColumnHeader", new IMAGE_PLUGIN_BRUSH_SVG("Icons/Package_16x", Icon16x16));
@@ -67,11 +70,32 @@ TSharedRef<FSlateStyleSet> FConcertServerStyle::Create()
 	StyleSet->Set("Concert.SessionContent.PackageDeleted", new IMAGE_PLUGIN_BRUSH_SVG("Icons/PackageDeleted_16x", Icon16x16));
 	StyleSet->Set("Concert.SessionContent.PackageRenamed", new IMAGE_PLUGIN_BRUSH_SVG("Icons/PackageRenamed_16x", Icon16x16));
 	StyleSet->Set("Concert.SessionContent.PackageSaved", new IMAGE_PLUGIN_BRUSH_SVG("Icons/PackageSaved_16x", Icon16x16));
+
+	// Clients tab
+	StyleSet->Set("Concert.Clients.DropShadow", new IMAGE_PLUGIN_BOX_BRUSH("ClientThumbnailDropShadow", FMargin(4.0f / 64.0f)));
+	StyleSet->Set("Concert.Clients.ThumbnailAreaHoverBackground", new FSlateRoundedBoxBrush(FStyleColors::Hover, 4.0f));
+	StyleSet->Set("Concert.Clients.ThumbnailAreaBackground", new FSlateRoundedBoxBrush(FStyleColors::Secondary, 4.0f));
+	StyleSet->Set("Concert.Clients.ThumbnailTitle", new FSlateRoundedBoxBrush(FStyleColors::Recessed, 4.0f));
+	StyleSet->Set("Concert.Clients.ThumbnailFooter", new FSlateRoundedBoxBrush(FStyleColors::Panel, 0.0f));
+	StyleSet->Set("Concert.Clients.TileTableRow", FTableRowStyle()
+			.SetEvenRowBackgroundBrush(FSlateNoResource() )
+			.SetEvenRowBackgroundHoveredBrush(FSlateNoResource())
+			.SetOddRowBackgroundBrush(FSlateNoResource())
+			.SetOddRowBackgroundHoveredBrush(FSlateNoResource())
+			.SetSelectorFocusedBrush(FSlateNoResource())
+			.SetActiveBrush(FSlateNoResource())
+			.SetActiveHoveredBrush(FSlateNoResource())
+			.SetInactiveBrush(FSlateNoResource())
+			.SetInactiveHoveredBrush(FSlateNoResource())
+			.SetTextColor(FSlateColor())
+			.SetSelectedTextColor(FSlateColor())
+			);
+	StyleSet->Set("Concert.Clients.ClientNameTileFont", DEFAULT_FONT("Regular", 16));
+	
 	return StyleSet;
 }
 
-#undef TODO_IMAGE_BRUSH
-#undef EDITOR_BOX_BRUSH
-#undef EDITOR_IMAGE_BRUSH_SVG
-#undef EDITOR_IMAGE_BRUSH
+#undef IMAGE_PLUGIN_BRUSH
+#undef IMAGE_PLUGIN_BOX_BRUSH
+#undef IMAGE_PLUGIN_BRUSH_SVG
 
diff --git a/Engine/Plugins/Developer/Concert/ConcertApp/MultiUserServer/Source/MultiUserServer/Private/Widgets/Clients/Browser/ClientBrowserItem.h b/Engine/Plugins/Developer/Concert/ConcertApp/MultiUserServer/Source/MultiUserServer/Private/Widgets/Clients/Browser/ClientBrowserItem.h
new file mode 100644
--- /dev/null
+++ b/Engine/Plugins/Developer/Concert/ConcertApp/MultiUserServer/Source/MultiUserServer/Private/Widgets/Clients/Browser/ClientBrowserItem.h
@@ -0,0 +1,21 @@
+// Copyright Epic Games, Inc. All Rights Reserved.
+
+#pragma once
+
+#include "CoreMinimal.h"
+#include "ConcertMessageData.h"
+#include "IMessageContext.h"
+
+namespace UE::MultiUserServer
+{
+	struct FClientBrowserItem
+	{
+		FConcertSessionClientInfo ClientInfo;
+		FMessageAddress ClientAddress;
+
+		FClientBrowserItem(const FConcertSessionClientInfo& ClientInfo, FMessageAddress ClientAddress)
+			: ClientInfo(ClientInfo)
+			, ClientAddress(ClientAddress)
+		{}
+	};
+}
diff --git a/Engine/Plugins/Developer/Concert/ConcertApp/MultiUserServer/Source/MultiUserServer/Private/Widgets/Clients/Browser/Models/ClientBrowserModel.cpp b/Engine/Plugins/Developer/Concert/ConcertApp/MultiUserServer/Source/MultiUserServer/Private/Widgets/Clients/Browser/Models/ClientBrowserModel.cpp
new file mode 100644
--- /dev/null
+++ b/Engine/Plugins/Developer/Concert/ConcertApp/MultiUserServer/Source/MultiUserServer/Private/Widgets/Clients/Browser/Models/ClientBrowserModel.cpp
@@ -0,0 +1,86 @@
+// Copyright Epic Games, Inc. All Rights Reserved.
+
+#include "ClientBrowserModel.h"
+
+#include "ConcertServerEvents.h"
+#include "ConcertUtil.h"
+#include "IConcertServer.h"
+#include "Algo/Transform.h"
+
+UE::MultiUserServer::FClientBrowserModel::FClientBrowserModel(TSharedRef<IConcertServer> InServer)
+	: Server(MoveTemp(InServer))
+{
+	ConcertServerEvents::OnLiveSessionCreated().AddRaw(this, &FClientBrowserModel::OnLiveSessionCreated);
+	ConcertServerEvents::OnLiveSessionDestroyed().AddRaw(this, &FClientBrowserModel::OnLiveSessionDestroyed);
+	for (const TSharedPtr<IConcertServerSession>& LiveSession : Server->GetLiveSessions())
+	{
+		SubscribeToClientConnectionEvents(LiveSession.ToSharedRef());
+	}
+}
+
+UE::MultiUserServer::FClientBrowserModel::~FClientBrowserModel()
+{
+	ConcertServerEvents::OnLiveSessionCreated().RemoveAll(this);
+	ConcertServerEvents::OnLiveSessionDestroyed().RemoveAll(this);
+	for (const TSharedPtr<IConcertServerSession>& LiveSession : Server->GetLiveSessions())
+	{
+		UnsubscribeFromClientConnectionEvents(LiveSession.ToSharedRef());
+	}
+}
+
+TSet<FGuid> UE::MultiUserServer::FClientBrowserModel::GetSessions() const
+{
+	const TArray<TSharedPtr<IConcertServerSession>> LiveSessions = Server->GetLiveSessions();
+	TSet<FGuid> Result;
+	Algo::Transform(LiveSessions, Result, [](const TSharedPtr<IConcertServerSession>& ServerSession){ return ServerSession->GetId(); });
+	return Result;
+}
+
+TOptional<FConcertSessionInfo> UE::MultiUserServer::FClientBrowserModel::GetSessionInfo(const FGuid& SessionID) const
+{
+	const TSharedPtr<IConcertServerSession> LiveSession = Server->GetLiveSession(SessionID);
+	return LiveSession
+		? LiveSession->GetSessionInfo()
+		: TOptional<FConcertSessionInfo>();
+}
+
+TArray<FConcertSessionClientInfo> UE::MultiUserServer::FClientBrowserModel::GetSessionClients(const FGuid& SessionID) const
+{
+	return ConcertUtil::GetSessionClients(*Server, SessionID);
+}
+
+FMessageAddress UE::MultiUserServer::FClientBrowserModel::GetClientAddress(const FGuid& ClientEndpointId) const
+{
+	const TSharedPtr<IConcertServerSession> ServerSession = ConcertUtil::GetLiveSessionClientConnectedTo(*Server, ClientEndpointId);
+	return ServerSession->GetClientAddress(ClientEndpointId);
+}
+
+void UE::MultiUserServer::FClientBrowserModel::OnLiveSessionCreated(bool bSuccess, const IConcertServer& InServer, TSharedRef<IConcertServerSession> InLiveSession) const
+{
+	if (bSuccess)
+	{
+		SubscribeToClientConnectionEvents(InLiveSession);
+	}
+	
+	OnSessionCreatedEvent.Broadcast(InLiveSession->GetId());
+}
+
+void UE::MultiUserServer::FClientBrowserModel::OnLiveSessionDestroyed(const IConcertServer& InServer, TSharedRef<IConcertServerSession> InLiveSession) const
+{
+	OnSessionCreatedEvent.Broadcast(InLiveSession->GetId());
+}
+
+void UE::MultiUserServer::FClientBrowserModel::SubscribeToClientConnectionEvents(const TSharedRef<IConcertServerSession>& InLiveSession) const
+{
+	InLiveSession->OnSessionClientChanged().AddRaw(this, &FClientBrowserModel::OnClientListUpdated);
+}
+
+void UE::MultiUserServer::FClientBrowserModel::UnsubscribeFromClientConnectionEvents( const TSharedRef<IConcertServerSession>& InLiveSession) const
+{
+	InLiveSession->OnSessionClientChanged().RemoveAll(this);
+}
+
+void UE::MultiUserServer::FClientBrowserModel::OnClientListUpdated(IConcertServerSession& Session, EConcertClientStatus Status, const FConcertSessionClientInfo& ClientInfo) const
+{
+	OnClientListChangedEvent.Broadcast(Session.GetId(), Status, ClientInfo);
+}
diff --git a/Engine/Plugins/Developer/Concert/ConcertApp/MultiUserServer/Source/MultiUserServer/Private/Widgets/Clients/Browser/Models/ClientBrowserModel.h b/Engine/Plugins/Developer/Concert/ConcertApp/MultiUserServer/Source/MultiUserServer/Private/Widgets/Clients/Browser/Models/ClientBrowserModel.h
new file mode 100644
--- /dev/null
+++ b/Engine/Plugins/Developer/Concert/ConcertApp/MultiUserServer/Source/MultiUserServer/Private/Widgets/Clients/Browser/Models/ClientBrowserModel.h
@@ -0,0 +1,47 @@
+// Copyright Epic Games, Inc. All Rights Reserved.
+
+#pragma once
+
+#include "CoreMinimal.h"
+#include "IClientBrowserModel.h"
+
+class IConcertServerSession;
+class IConcertServer;
+
+namespace UE::MultiUserServer
+{
+	class FClientBrowserModel : public IClientBrowserModel
+	{
+	public:
+
+		FClientBrowserModel(TSharedRef<IConcertServer> InServer);
+		virtual ~FClientBrowserModel() override;
+
+		//~ Begin IClientBrowserModel Interface
+		virtual TSet<FGuid> GetSessions() const override;
+		virtual TOptional<FConcertSessionInfo> GetSessionInfo(const FGuid& SessionID) const override;
+		virtual TArray<FConcertSessionClientInfo> GetSessionClients(const FGuid& SessionID) const override;
+		virtual FMessageAddress GetClientAddress(const FGuid& ClientEndpointId) const override;
+		virtual FOnClientListChanged& OnClientListChanged() override { return OnClientListChangedEvent; }
+		virtual FOnSessionListChanged& OnSessionCreated() override { return OnSessionCreatedEvent; }
+		virtual FOnSessionListChanged& OnSessionDestroyed() override { return OnSessionDestroyedEvent; }
+		//~ End IClientBrowserModel Interface
+
+	private:
+
+		TSharedRef<IConcertServer> Server;
+
+		FOnClientListChanged OnClientListChangedEvent;
+		FOnSessionListChanged OnSessionCreatedEvent;
+		FOnSessionListChanged OnSessionDestroyedEvent;
+
+		void OnLiveSessionCreated(bool bSuccess, const IConcertServer& InServer, TSharedRef<IConcertServerSession> InLiveSession) const;
+		void OnLiveSessionDestroyed(const IConcertServer& InServer, TSharedRef<IConcertServerSession> InLiveSession) const;
+		
+		void SubscribeToClientConnectionEvents(const TSharedRef<IConcertServerSession>& InLiveSession) const;
+		void UnsubscribeFromClientConnectionEvents(const TSharedRef<IConcertServerSession>& InLiveSession) const;
+		void OnClientListUpdated(IConcertServerSession& Session, EConcertClientStatus Status, const FConcertSessionClientInfo& ClientInfo) const;
+	};
+}
+
+
diff --git a/Engine/Plugins/Developer/Concert/ConcertApp/MultiUserServer/Source/MultiUserServer/Private/Widgets/Clients/Browser/Models/ClientNetworkStatisticsModel.cpp b/Engine/Plugins/Developer/Concert/ConcertApp/MultiUserServer/Source/MultiUserServer/Private/Widgets/Clients/Browser/Models/ClientNetworkStatisticsModel.cpp
new file mode 100644
--- /dev/null
+++ b/Engine/Plugins/Developer/Concert/ConcertApp/MultiUserServer/Source/MultiUserServer/Private/Widgets/Clients/Browser/Models/ClientNetworkStatisticsModel.cpp
@@ -0,0 +1,100 @@
+// Copyright Epic Games, Inc. All Rights Reserved.
+
+#include "ClientNetworkStatisticsModel.h"
+
+#include "Containers/Ticker.h"
+#include "Features/IModularFeatures.h"
+#include "INetworkMessagingExtension.h"
+#include "Misc/ScopeExit.h"
+
+namespace UE::MultiUserServer::Private
+{
+	static INetworkMessagingExtension* GetMessagingStatistics()
+	{
+		IModularFeatures& ModularFeatures = IModularFeatures::Get();
+		if (IsInGameThread())
+		{
+			if (ModularFeatures.IsModularFeatureAvailable(INetworkMessagingExtension::ModularFeatureName))
+			{
+				return &ModularFeatures.GetModularFeature<INetworkMessagingExtension>(INetworkMessagingExtension::ModularFeatureName);
+			}
+		}
+		else
+		{
+			ModularFeatures.LockModularFeatureList();
+			ON_SCOPE_EXIT
+			{
+				ModularFeatures.UnlockModularFeatureList();
+			};
+			
+			if (ModularFeatures.IsModularFeatureAvailable(INetworkMessagingExtension::ModularFeatureName))
+			{
+				return &ModularFeatures.GetModularFeature<INetworkMessagingExtension>(INetworkMessagingExtension::ModularFeatureName);
+			}
+		}
+		
+
+		ensureMsgf(false, TEXT("Feature %s is unavailable"), *INetworkMessagingExtension::ModularFeatureName.ToString());
+		return nullptr;
+	}
+}
+
+UE::MultiUserServer::FClientNetworkStatisticsModel::FClientNetworkStatisticsModel()
+{
+	if (INetworkMessagingExtension* Statistics = Private::GetMessagingStatistics())
+	{
+		Statistics->OnTransferUpdatedFromThread().AddRaw(this, &FClientNetworkStatisticsModel::OnTransferUpdatedFromThread);
+	}
+}
+
+UE::MultiUserServer::FClientNetworkStatisticsModel::~FClientNetworkStatisticsModel()
+{
+	if (INetworkMessagingExtension* Statistics = Private::GetMessagingStatistics())
+	{
+		Statistics->OnTransferUpdatedFromThread().RemoveAll(this);
+	}
+}
+
+TOptional<FMessageTransportStatistics> UE::MultiUserServer::FClientNetworkStatisticsModel::GetLatestNetworkStatistics(const FMessageAddress& ClientAddress) const
+{
+	if (INetworkMessagingExtension* Statistics = Private::GetMessagingStatistics())
+	{
+		const FGuid NodeId = Statistics->GetNodeIdFromAddress(ClientAddress);
+		return NodeId.IsValid() ? Statistics->GetLatestNetworkStatistics(NodeId) : TOptional<FMessageTransportStatistics>();
+	}
+	return {};
+}
+
+void UE::MultiUserServer::FClientNetworkStatisticsModel::RegisterOnTransferUpdatedFromThread(const FMessageAddress& ClientAddress, FOnMessageTransportStatisticsUpdated StatisticUpdateCallback)
+{
+	INetworkMessagingExtension* Statistics = Private::GetMessagingStatistics();
+	if (!Statistics || !ensure(IsInGameThread()))
+	{
+		return;
+	}
+	if (const FGuid NodeId = Statistics->GetNodeIdFromAddress(ClientAddress); NodeId.IsValid())
+	{
+		StatisticUpdateCallbacks.Add(NodeId, StatisticUpdateCallback);
+	}
+}
+
+void UE::MultiUserServer::FClientNetworkStatisticsModel::UnregisterOnTransferUpdatedFromThread(const FMessageAddress& ClientAddress)
+{
+	INetworkMessagingExtension* Statistics = Private::GetMessagingStatistics();
+	if (!Statistics || !ensure(IsInGameThread()))
+	{
+		return;
+	}
+	if (const FGuid NodeId = Statistics->GetNodeIdFromAddress(ClientAddress); NodeId.IsValid())
+	{
+		StatisticUpdateCallbacks.Remove(NodeId);
+	}
+}
+
+void UE::MultiUserServer::FClientNetworkStatisticsModel::OnTransferUpdatedFromThread(FTransferStatistics Stats) const
+{
+	if (const FOnMessageTransportStatisticsUpdated* Callback = StatisticUpdateCallbacks.Find(Stats.DestinationId))
+	{
+		Callback->Execute(Stats);
+	}
+}
\ No newline at end of file
diff --git a/Engine/Plugins/Developer/Concert/ConcertApp/MultiUserServer/Source/MultiUserServer/Private/Widgets/Clients/Browser/Models/ClientNetworkStatisticsModel.h b/Engine/Plugins/Developer/Concert/ConcertApp/MultiUserServer/Source/MultiUserServer/Private/Widgets/Clients/Browser/Models/ClientNetworkStatisticsModel.h
new file mode 100644
--- /dev/null
+++ b/Engine/Plugins/Developer/Concert/ConcertApp/MultiUserServer/Source/MultiUserServer/Private/Widgets/Clients/Browser/Models/ClientNetworkStatisticsModel.h
@@ -0,0 +1,31 @@
+// Copyright Epic Games, Inc. All Rights Reserved.
+
+#pragma once
+
+#include "CoreMinimal.h"
+#include "IClientNetworkStatisticsModel.h"
+
+namespace UE::MultiUserServer
+{
+	/** Synchronizes the network statistics (the statistics are updated async). */
+	class FClientNetworkStatisticsModel : public IClientNetworkStatisticsModel
+	{
+	public:
+
+		FClientNetworkStatisticsModel();
+		virtual ~FClientNetworkStatisticsModel() override;
+
+		//~ Begin IClientNetworkStatisticsModel Interface
+		virtual TOptional<FMessageTransportStatistics> GetLatestNetworkStatistics(const FMessageAddress& ClientAddress) const override;
+		virtual void RegisterOnTransferUpdatedFromThread(const FMessageAddress& ClientAddress, FOnMessageTransportStatisticsUpdated StatisticsUpdatedCallback) override;
+		virtual void UnregisterOnTransferUpdatedFromThread(const FMessageAddress& ClientAddress) override;
+		//~ End IClientNetworkStatisticsModel Interface
+
+	private:
+
+		TMap<FGuid, FOnMessageTransportStatisticsUpdated> StatisticUpdateCallbacks;
+
+		void OnTransferUpdatedFromThread(FTransferStatistics Stats) const;
+	};
+}
+
diff --git a/Engine/Plugins/Developer/Concert/ConcertApp/MultiUserServer/Source/MultiUserServer/Private/Widgets/Clients/Browser/Models/IClientBrowserModel.h b/Engine/Plugins/Developer/Concert/ConcertApp/MultiUserServer/Source/MultiUserServer/Private/Widgets/Clients/Browser/Models/IClientBrowserModel.h
new file mode 100644
--- /dev/null
+++ b/Engine/Plugins/Developer/Concert/ConcertApp/MultiUserServer/Source/MultiUserServer/Private/Widgets/Clients/Browser/Models/IClientBrowserModel.h
@@ -0,0 +1,39 @@
+// Copyright Epic Games, Inc. All Rights Reserved.
+
+#pragma once
+
+#include "CoreMinimal.h"
+#include "ConcertMessages.h"
+
+struct FConcertSessionClientInfo;
+struct FConcertSessionInfo;
+struct FMessageAddress;
+
+namespace UE::MultiUserServer
+{
+	/** Decouples the UI from the server functions. */
+	class IClientBrowserModel
+	{
+	public:
+
+		/** Gets the IDs of all availabel IDs */
+		virtual TSet<FGuid> GetSessions() const = 0;
+		/** Gets more info about a session returned by GetSessions */
+		virtual TOptional<FConcertSessionInfo> GetSessionInfo(const FGuid& SessionID) const = 0;
+		/** Gets the clients connected to a session returned by GetSessions */
+		virtual TArray<FConcertSessionClientInfo> GetSessionClients(const FGuid& SessionID) const = 0;
+		/** Gets the network address of a given client */
+		virtual FMessageAddress GetClientAddress(const FGuid& ClientEndpointId) const = 0; 
+		
+		DECLARE_MULTICAST_DELEGATE_ThreeParams(FOnClientListChanged, const FGuid& /*SessionId*/, EConcertClientStatus /*UpdateType*/, const FConcertSessionClientInfo& /*ClientInfo*/);
+		virtual FOnClientListChanged& OnClientListChanged() = 0;
+
+		DECLARE_MULTICAST_DELEGATE_OneParam(FOnSessionListChanged, const FGuid& /*SessionId*/);
+		virtual FOnSessionListChanged& OnSessionCreated() = 0;
+		virtual FOnSessionListChanged& OnSessionDestroyed() = 0;
+		
+		virtual ~IClientBrowserModel() = default;
+	};
+}
+
+
diff --git a/Engine/Plugins/Developer/Concert/ConcertApp/MultiUserServer/Source/MultiUserServer/Private/Widgets/Clients/Browser/Models/IClientNetworkStatisticsModel.h b/Engine/Plugins/Developer/Concert/ConcertApp/MultiUserServer/Source/MultiUserServer/Private/Widgets/Clients/Browser/Models/IClientNetworkStatisticsModel.h
new file mode 100644
--- /dev/null
+++ b/Engine/Plugins/Developer/Concert/ConcertApp/MultiUserServer/Source/MultiUserServer/Private/Widgets/Clients/Browser/Models/IClientNetworkStatisticsModel.h
@@ -0,0 +1,31 @@
+// Copyright Epic Games, Inc. All Rights Reserved.
+
+#pragma once
+
+#include "CoreMinimal.h"
+
+struct FConcertSessionClientInfo;
+struct FConcertSessionInfo;
+struct FMessageAddress;
+struct FMessageTransportStatistics;
+struct FTransferStatistics;
+
+namespace UE::MultiUserServer
+{
+	DECLARE_DELEGATE_OneParam(FOnMessageTransportStatisticsUpdated, const FTransferStatistics&);
+	
+	/** Decouples the UI from the server functions. */
+	class IClientNetworkStatisticsModel
+	{
+	public:
+		
+		virtual TOptional<FMessageTransportStatistics> GetLatestNetworkStatistics(const FMessageAddress& ClientAddress) const = 0;
+
+		virtual void RegisterOnTransferUpdatedFromThread(const FMessageAddress& ClientAddress, FOnMessageTransportStatisticsUpdated StatisticsUpdatedCallback) = 0;
+		virtual void UnregisterOnTransferUpdatedFromThread(const FMessageAddress& ClientAddress) = 0;
+
+		virtual ~IClientNetworkStatisticsModel() = default;
+	};
+}
+
+
diff --git a/Engine/Plugins/Developer/Concert/ConcertApp/MultiUserServer/Source/MultiUserServer/Private/Widgets/Clients/Browser/SClientNetworkStats.cpp b/Engine/Plugins/Developer/Concert/ConcertApp/MultiUserServer/Source/MultiUserServer/Private/Widgets/Clients/Browser/SClientNetworkStats.cpp
new file mode 100644
--- /dev/null
+++ b/Engine/Plugins/Developer/Concert/ConcertApp/MultiUserServer/Source/MultiUserServer/Private/Widgets/Clients/Browser/SClientNetworkStats.cpp
@@ -0,0 +1,151 @@
+// Copyright Epic Games, Inc. All Rights Reserved.
+
+#include "SClientNetworkStats.h"
+
+#include "ConcertServerStyle.h"
+#include "INetworkMessagingExtension.h"
+#include "Math/UnitConversion.h"
+#include "Widgets/Clients/Browser/Models/ClientNetworkStatisticsModel.h"
+
+#include "Widgets/Layout/SBorder.h"
+#include "Widgets/SBoxPanel.h"
+#include "Widgets/Text/STextBlock.h"
+
+#define LOCTEXT_NAMESPACE "UnrealMultiUserUI.SClientNetworkStats"
+
+namespace UE::MultiUserServer
+{
+	void SClientNetworkStats::Construct(const FArguments& InArgs, const FMessageAddress& InNodeAddress, TSharedRef<IClientNetworkStatisticsModel> InStatisticModel)
+	{
+		NodeAddress = InNodeAddress;
+		StatisticModel = MoveTemp(InStatisticModel);
+		
+		HighlightText = InArgs._HighlightText;
+		check(HighlightText);
+		
+		ChildSlot
+		.HAlign(HAlign_Fill)
+		[
+			CreateStatTable()
+		];
+	}
+
+	void SClientNetworkStats::Tick(const FGeometry& AllottedGeometry, const double InCurrentTime,const float InDeltaTime)
+	{
+		SCompoundWidget::Tick(AllottedGeometry, InCurrentTime, InDeltaTime);
+
+		if (const TOptional<FMessageTransportStatistics> Stats = StatisticModel->GetLatestNetworkStatistics(NodeAddress))
+		{
+			UpdateStatistics(*Stats);
+		}
+		else
+		{
+			const FText NotApplicableText = FText::FromString(TEXT("n/a"));
+			SendText->SetText(NotApplicableText);
+			ReceiveText->SetText(NotApplicableText);
+			RoundTripTimeText->SetText(NotApplicableText);
+			InflightText->SetText(NotApplicableText);
+			LossText->SetText(NotApplicableText);
+		}
+	}
+
+	void SClientNetworkStats::AppendSearchTerms(TArray<FString>& SearchTerms) const
+	{
+		SearchTerms.Add(SendText->GetText().ToString());
+		SearchTerms.Add(ReceiveText->GetText().ToString());
+		SearchTerms.Add(RoundTripTimeText->GetText().ToString());
+		SearchTerms.Add(InflightText->GetText().ToString());
+		SearchTerms.Add(LossText->GetText().ToString());
+	}
+
+	void SClientNetworkStats::UpdateStatistics(const FMessageTransportStatistics& Statistics)
+	{
+		// Send
+		{
+			const FNumericUnit<uint64> Unit = FUnitConversion::QuantizeUnitsToBestFit(Statistics.TotalBytesSent, EUnit::Bytes);
+			const FString DisplayString = FString::Printf(TEXT("%lld %s"), Unit.Value, FUnitConversion::GetUnitDisplayString(Unit.Units));
+			SendText->SetText(FText::FromString(DisplayString));
+		}
+		// Receive
+		{
+			const FNumericUnit<uint64> Unit = FUnitConversion::QuantizeUnitsToBestFit(Statistics.TotalBytesReceived, EUnit::Bytes);
+			const FString DisplayString = FString::Printf(TEXT("%lld %s"), Statistics.TotalBytesReceived, FUnitConversion::GetUnitDisplayString(Unit.Units));
+			ReceiveText->SetText(FText::FromString(DisplayString));
+		}
+		// RTT
+		{
+			const FNumericUnit<uint64> Unit = FUnitConversion::QuantizeUnitsToBestFit(static_cast<uint64>(Statistics.AverageRTT.GetTotalMilliseconds()), EUnit::Milliseconds);
+			const FString DisplayString = FString::Printf(TEXT("%lld %s"), Unit.Value, FUnitConversion::GetUnitDisplayString(Unit.Units));
+			RoundTripTimeText->SetText(FText::FromString(DisplayString));
+		}
+		// Inflight
+		{
+			const FNumericUnit<uint64> Unit = FUnitConversion::QuantizeUnitsToBestFit(Statistics.BytesInflight, EUnit::Bytes);
+			const FString DisplayString = FString::Printf(TEXT("%lld %s"), Statistics.BytesInflight, FUnitConversion::GetUnitDisplayString(Unit.Units));
+			InflightText->SetText(FText::FromString(DisplayString));
+		}
+		// Loss
+		{
+			const FNumericUnit<uint64> Unit = FUnitConversion::QuantizeUnitsToBestFit(Statistics.TotalBytesLost, EUnit::Bytes);
+			const FString DisplayString = FString::Printf(TEXT("%lld %s"), Statistics.TotalBytesLost, FUnitConversion::GetUnitDisplayString(Unit.Units));
+			LossText->SetText(FText::FromString(DisplayString));
+		}
+	}
+
+	TSharedRef<SWidget> SClientNetworkStats::CreateStatTable()
+	{
+		const TSharedRef<SHorizontalBox> Content = SNew(SHorizontalBox)
+			.ToolTipText_Lambda([this]()
+			{
+				const TOptional<FMessageTransportStatistics> Stats = StatisticModel->GetLatestNetworkStatistics(NodeAddress);
+				return Stats.IsSet() ? FText::GetEmpty() : LOCTEXT("ErrorGettingStats", "Error getting stats");
+			});
+		AddStatistic(Content, LOCTEXT("SentLabel", "Sent"), LOCTEXT("SentTooltip", "Total bytes sent to this client"), SendText);
+		AddStatistic(Content, LOCTEXT("ReceiveLabel", "Received"), LOCTEXT("ReceiveTooltip", "Total bytes received from this client"), ReceiveText);
+		AddStatistic(Content, LOCTEXT("RttLabel", "RTT"), LOCTEXT("RttTooltip", "Round trip time"), RoundTripTimeText);
+		AddStatistic(Content, LOCTEXT("InflightLabel", "Inflight"), LOCTEXT("InflightTooltip", "Total reliable bytes awaiting an ack from client"), InflightText);
+
+		// Loss should be at the right
+		SHorizontalBox::FScopedWidgetSlotArguments LossSlat = Content->AddSlot();
+		AddStatistic(LossSlat, LOCTEXT("LossLabel", "Lost"), LOCTEXT("LostTooltip", "Total bytes lost while sending to the client"), LossText);
+		LossSlat.HAlign(HAlign_Right);
+		LossSlat.FillWidth(1.f);
+		
+		return Content;
+	}
+	
+	void SClientNetworkStats::AddStatistic(const TSharedRef<SHorizontalBox>& AddTo, const FText& StatisticName, const FText& StatisticToolTip, TSharedPtr<STextBlock>& AssignTo)
+	{
+		SHorizontalBox::FScopedWidgetSlotArguments Slot = AddTo->AddSlot();
+		AddStatistic(Slot, StatisticName, StatisticToolTip, AssignTo);
+	}
+
+	void SClientNetworkStats::AddStatistic(SHorizontalBox::FScopedWidgetSlotArguments& Slot, const FText& StatisticName, const FText& StatisticToolTip, TSharedPtr<STextBlock>& AssignTo)
+	{
+		Slot
+		.AutoWidth()
+		.Padding(3.f)
+		[
+			SNew(SVerticalBox)
+			.ToolTipText(StatisticToolTip)
+
+			+SVerticalBox::Slot()
+			.AutoHeight()
+			[
+				SNew(STextBlock)
+				.Text(StatisticName)
+				.ColorAndOpacity(FColor::White)
+			]
+
+			+SVerticalBox::Slot()
+			.AutoHeight()
+			[
+				SAssignNew(AssignTo, STextBlock)
+				.ColorAndOpacity(FColor::White)
+				.HighlightText_Lambda([this](){ return *HighlightText; })
+			]
+		];
+	}
+}
+
+#undef LOCTEXT_NAMESPACE
\ No newline at end of file
diff --git a/Engine/Plugins/Developer/Concert/ConcertApp/MultiUserServer/Source/MultiUserServer/Private/Widgets/Clients/Browser/SClientNetworkStats.h b/Engine/Plugins/Developer/Concert/ConcertApp/MultiUserServer/Source/MultiUserServer/Private/Widgets/Clients/Browser/SClientNetworkStats.h
new file mode 100644
--- /dev/null
+++ b/Engine/Plugins/Developer/Concert/ConcertApp/MultiUserServer/Source/MultiUserServer/Private/Widgets/Clients/Browser/SClientNetworkStats.h
@@ -0,0 +1,61 @@
+// Copyright Epic Games, Inc. All Rights Reserved.
+
+#pragma once
+
+#include "CoreMinimal.h"
+#include "IMessageContext.h"
+#include "Widgets/DeclarativeSyntaxSupport.h"
+#include "Widgets/SBoxPanel.h"
+#include "Widgets/SCompoundWidget.h"
+
+class STextBlock;
+struct FMessageTransportStatistics;
+
+namespace UE::MultiUserServer
+{
+	class IClientNetworkStatisticsModel;
+	
+	/**
+	 * Displays statistics about a client connection in a table like format: send, receive, RTT, Inflight, and Loss.
+	 */
+	class SClientNetworkStats : public SCompoundWidget
+	{
+	public:
+
+		SLATE_BEGIN_ARGS(SClientNetworkStats){}
+			SLATE_ARGUMENT(TSharedPtr<FText>, HighlightText)
+		SLATE_END_ARGS()
+
+		void Construct(const FArguments& InArgs, const FMessageAddress& InNodeAddress, TSharedRef<IClientNetworkStatisticsModel> InStatisticModel);
+
+		//~ Begin SCompoundWidget Interface		
+		virtual void Tick(const FGeometry& AllottedGeometry, const double InCurrentTime, const float InDeltaTime) override;
+		//~ End SCompoundWidget Interface
+		
+		/** Appends the statistics to the search terms */
+		void AppendSearchTerms(TArray<FString>& SearchTerms) const;
+
+	private:
+
+		/** The ID getting visualised */
+		FMessageAddress NodeAddress;
+		/** Used to obtain settings */
+		TSharedPtr<IClientNetworkStatisticsModel> StatisticModel;
+		
+		/** The text to highlight */
+		TSharedPtr<FText> HighlightText;
+
+		TSharedPtr<STextBlock> SendText;
+		TSharedPtr<STextBlock> ReceiveText;
+		TSharedPtr<STextBlock> RoundTripTimeText;
+		TSharedPtr<STextBlock> InflightText;
+		TSharedPtr<STextBlock> LossText;
+		
+		void UpdateStatistics(const FMessageTransportStatistics& Statistics);
+
+		TSharedRef<SWidget> CreateStatTable();
+		void AddStatistic(const TSharedRef<SHorizontalBox>& AddTo, const FText& StatisticName, const FText& StatisticToolTip, TSharedPtr<STextBlock>& AssignTo);
+		void AddStatistic(SHorizontalBox::FScopedWidgetSlotArguments& Slot, const FText& StatisticName, const FText& StatisticToolTip, TSharedPtr<STextBlock>& AssignTo);
+	};
+}
+
diff --git a/Engine/Plugins/Developer/Concert/ConcertApp/MultiUserServer/Source/MultiUserServer/Private/Widgets/Clients/Browser/SConcertClientBrowser.cpp b/Engine/Plugins/Developer/Concert/ConcertApp/MultiUserServer/Source/MultiUserServer/Private/Widgets/Clients/Browser/SConcertClientBrowser.cpp
new file mode 100644
--- /dev/null
+++ b/Engine/Plugins/Developer/Concert/ConcertApp/MultiUserServer/Source/MultiUserServer/Private/Widgets/Clients/Browser/SConcertClientBrowser.cpp
@@ -0,0 +1,362 @@
+// Copyright Epic Games, Inc. All Rights Reserved.
+
+#include "SConcertClientBrowser.h"
+
+#include "ClientBrowserItem.h"
+#include "ConcertServerStyle.h"
+#include "SConcertClientBrowserItem.h"
+#include "Models/IClientBrowserModel.h"
+
+#include "Algo/AnyOf.h"
+#include "Framework/MultiBox/MultiBoxBuilder.h"
+#include "Models/IClientNetworkStatisticsModel.h"
+#include "Widgets/SBoxPanel.h"
+#include "Widgets/Input/SComboBox.h"
+#include "Widgets/Input/SSearchBox.h"
+#include "Widgets/Views/STileView.h"
+
+#define LOCTEXT_NAMESPACE "UnrealMultiUserUI.SConcertClientBrowser"
+
+void UE::MultiUserServer::SConcertClientBrowser::Construct(const FArguments& InArgs, TSharedRef<IClientBrowserModel> InBrowserModel, TSharedRef<IClientNetworkStatisticsModel> InStatisticsModel)
+{
+	BrowserModel = MoveTemp(InBrowserModel);
+	StatisticsModel = MoveTemp(InStatisticsModel);
+	
+	HighlightText = MakeShared<FText>();
+	SessionFilter = MakeShared<FClientTextFilter>(FClientTextFilter::FItemToStringArray::CreateSP(this, &SConcertClientBrowser::GenerateSearchTerms));
+	SessionFilter->OnChanged().AddSP(this, &SConcertClientBrowser::UpdateTileViewFromAllowedSessions);
+
+	OnClientDoubleClicked = InArgs._OnClientDoubleClicked;
+
+	ChildSlot
+	[
+		SNew(SVerticalBox)
+
+		+SVerticalBox::Slot()
+		.AutoHeight()
+		[
+			CreateSearchArea(InArgs)
+		]
+
+		+SVerticalBox::Slot()
+		.FillHeight(1.f)
+		.Padding(5.f)
+		[
+			SNew(SOverlay)
+
+			+SOverlay::Slot()
+			[
+				CreateTileView()
+			]
+
+			+SOverlay::Slot()
+			[
+				SNew(SBox)
+				.HAlign(HAlign_Center)
+				.VAlign(VAlign_Center)
+				.Visibility_Lambda([this](){ return DisplayedClients.Num() == 0 ? EVisibility::Visible : EVisibility::Hidden; })
+				[
+					SNew(STextBlock)
+					.Text(this, &SConcertClientBrowser::GetErrorMessageText)
+				]
+			]
+		]
+	];
+
+	BrowserModel->OnSessionCreated().AddSP(this, &SConcertClientBrowser::OnSessionCreated);
+	BrowserModel->OnSessionDestroyed().AddSP(this, &SConcertClientBrowser::OnSessionDestroyed);
+	BrowserModel->OnClientListChanged().AddSP(this, &SConcertClientBrowser::OnClientListChanged);
+
+	AllowAllSessions();
+}
+
+void UE::MultiUserServer::SConcertClientBrowser::ShowOnlyClientsFromSession(const FGuid& SessionId)
+{
+	AllowedSessions = { SessionId };
+	UpdateTileViewFromAllowedSessions();
+}
+
+TSharedRef<SWidget> UE::MultiUserServer::SConcertClientBrowser::CreateSearchArea(const FArguments& InArgs)
+{
+	return SNew(SHorizontalBox)
+	
+		+SHorizontalBox::Slot()
+		.AutoWidth()
+		[
+			SNew(SComboButton)
+			.OnGetMenuContent(this, &SConcertClientBrowser::MakeSessionOption)
+			.ButtonContent()
+			[
+				SNew(STextBlock)
+				.Text_Lambda([this]()
+				{
+					return AllowedSessions.Num() == BrowserModel->GetSessions().Num()
+						? LOCTEXT("MessageActionFilter.Selection.All", "All")
+						: FText::FromString(FString::FromInt(AllowedSessions.Num()));
+				})
+			]
+		]
+
+		+SHorizontalBox::Slot()
+		.AutoWidth()
+		.FillWidth(0.65f)
+		.Padding(4.f, 0.f)
+		.VAlign(VAlign_Fill)
+		[
+			SNew(SSearchBox)
+			.OnTextChanged_Lambda([this](const FText& SearchText)
+			{
+				*HighlightText = SearchText;
+				SessionFilter->SetRawFilterText(SearchText);
+			})
+		]
+
+		+SHorizontalBox::Slot()
+		.AutoWidth()
+		[
+			InArgs._RightOfSearch.Widget
+		];
+}
+
+TSharedRef<SWidget> UE::MultiUserServer::SConcertClientBrowser::CreateTileView()
+{
+	constexpr int32 Height = 200;
+	constexpr int32 Width = 270;
+	return SAssignNew(TileView, STileView<TSharedPtr<FClientBrowserItem>>)
+		.SelectionMode(ESelectionMode::Multi)
+		.ListItemsSource(&DisplayedClients)
+		.OnGenerateTile(this, &SConcertClientBrowser::MakeTileViewWidget)
+		.OnContextMenuOpening(this, &SConcertClientBrowser::OnGetContextMenuContent)
+		.OnMouseButtonDoubleClick(this, &SConcertClientBrowser::OnListMouseButtonDoubleClick)
+		.ItemHeight(Height)
+		.ItemWidth(Width)
+	;
+}
+
+void UE::MultiUserServer::SConcertClientBrowser::OnSessionCreated(const FGuid& SessionId)
+{
+	if (bShowAllSessions)
+	{
+		AllowedSessions.Add(SessionId);
+		UpdateTileViewFromAllowedSessions();
+	}
+}
+
+void UE::MultiUserServer::SConcertClientBrowser::OnSessionDestroyed(const FGuid& SessionId)
+{
+	if (AllowedSessions.Contains(SessionId))
+	{
+		AllowedSessions.Remove(SessionId);
+		UpdateTileViewFromAllowedSessions();
+	}
+}
+
+void UE::MultiUserServer::SConcertClientBrowser::OnClientListChanged(const FGuid& SessionId, EConcertClientStatus UpdateType, const FConcertSessionClientInfo& ClientInfo)
+{
+	switch (UpdateType)
+	{
+	case EConcertClientStatus::Connected:
+		DisplayedClients.Add(MakeShared<FClientBrowserItem>(ClientInfo, BrowserModel->GetClientAddress(ClientInfo.ClientEndpointId)));
+		TileView->RequestListRefresh();
+		break;
+	case EConcertClientStatus::Disconnected:
+		RemoveClient(ClientInfo);
+		break;
+	case EConcertClientStatus::Updated:
+		UpdateClientInfo(ClientInfo);
+		break;
+	default: ;
+	}
+}
+
+void UE::MultiUserServer::SConcertClientBrowser::RemoveClient(const FConcertSessionClientInfo& ClientInfo)
+{
+	const int32 Index = DisplayedClients.IndexOfByPredicate([&ClientInfo](const TSharedPtr<FClientBrowserItem>& Item)
+		{
+			return Item->ClientInfo.ClientEndpointId == ClientInfo.ClientEndpointId;
+		});
+	if (ensure(Index != INDEX_NONE))
+	{
+		DisplayedClients.RemoveAt(Index);
+		ClientWidgets.Remove(ClientInfo.ClientEndpointId);
+		TileView->RequestListRefresh();
+	}
+}
+
+void UE::MultiUserServer::SConcertClientBrowser::UpdateClientInfo(const FConcertSessionClientInfo& ClientInfo)
+{
+	if (const TSharedPtr<SConcertClientBrowserItem>* Widget = ClientWidgets.Find(ClientInfo.ClientEndpointId))
+	{
+		Widget->Get()->OnClientInfoChanged();
+	}
+}
+
+TSharedRef<SWidget> UE::MultiUserServer::SConcertClientBrowser::SConcertClientBrowser::MakeSessionOption()
+{
+	FMenuBuilder MenuBuilder(true, nullptr);
+
+	MenuBuilder.AddMenuEntry(
+		LOCTEXT("SelectAll", "All"),
+		FText::GetEmpty(),
+		FSlateIcon(),
+		FUIAction(
+			FExecuteAction::CreateLambda([this]()
+			{
+				if (BrowserModel->GetSessions().Num() == AllowedSessions.Num())
+				{
+					DisallowAllSessions();
+				}
+				else
+				{
+					AllowAllSessions();
+				}
+			}),
+			FCanExecuteAction::CreateLambda([] { return true; }),
+			FIsActionChecked::CreateLambda([this](){ return BrowserModel->GetSessions().Num() == AllowedSessions.Num(); })),
+		NAME_None,
+		EUserInterfaceActionType::Check
+	);
+	MenuBuilder.AddSeparator();
+	for (const FGuid& SessionId : BrowserModel->GetSessions())
+	{
+		MenuBuilder.AddMenuEntry(
+			FText::FromString(BrowserModel->GetSessionInfo(SessionId)->SessionName),
+			FText::GetEmpty(),
+			FSlateIcon(),
+			FUIAction(
+				FExecuteAction::CreateLambda([this, SessionId]()
+				{
+					if (AllowedSessions.Contains(SessionId))
+					{
+						AllowedSessions.Remove(SessionId);
+					}
+					else
+					{
+						AllowedSessions.Add(SessionId);
+					}
+					UpdateTileViewFromAllowedSessions();
+				}),
+				FCanExecuteAction::CreateLambda([] { return true; }),
+				FIsActionChecked::CreateLambda([this, SessionId]{ return AllowedSessions.Contains(SessionId); })),
+			NAME_None,
+			EUserInterfaceActionType::Check
+		);
+	}
+
+	return MenuBuilder.MakeWidget();
+}
+
+FText UE::MultiUserServer::SConcertClientBrowser::GetErrorMessageText() const
+{
+	if (BrowserModel->GetSessions().Num() == 0)
+	{
+		return LOCTEXT("NoLiveSessions", "No live sessions");
+	}
+	
+	if (AllowedSessions.Num() == 0)
+	{
+		return LOCTEXT("NoSessionsSelected", "All live sessions filtered out");
+	}
+
+	const bool bAtLeastOneClient = Algo::AnyOf(AllowedSessions, [this](const FGuid& SessionId)
+	{
+		return BrowserModel->GetSessionClients(SessionId).Num() > 0;
+	});
+	if (bAtLeastOneClient)
+	{
+		return LOCTEXT("AllFilteredOut", "All results have been filtered. Try changing your active filters above.");
+	}
+	
+	return LOCTEXT("NoClients", "No clients connected to selected sessions");
+}
+
+TSharedRef<ITableRow> UE::MultiUserServer::SConcertClientBrowser::MakeTileViewWidget(TSharedPtr<FClientBrowserItem> ClientItem, const TSharedRef<STableViewBase>& OwnerTable)
+{
+	const FClientEndpointId ClientId = ClientItem->ClientInfo.ClientEndpointId;
+	TSharedPtr<SConcertClientBrowserItem>* ExistingWidget = ClientWidgets.Find(ClientId);
+	if (!ExistingWidget)
+	{
+		ClientWidgets.Add(ClientItem->ClientInfo.ClientEndpointId, SNew(SConcertClientBrowserItem, ClientItem.ToSharedRef(), StatisticsModel.ToSharedRef())
+			 .HighlightText(HighlightText)
+			 );
+	}
+	
+	return SNew(STableRow<TSharedPtr<FClientBrowserItem>>, OwnerTable )
+		.Padding(2.f)
+		.Style(FConcertServerStyle::Get(), "Concert.Clients.TileTableRow")
+		.Content()
+		[
+			ClientWidgets[ClientId].ToSharedRef()
+		];
+}
+
+TSharedPtr<SWidget> UE::MultiUserServer::SConcertClientBrowser::OnGetContextMenuContent()
+{
+	FMenuBuilder MenuBuilder(true, nullptr);
+
+	MenuBuilder.AddMenuEntry(
+		LOCTEXT("OpenLog", "Open log"),
+		LOCTEXT("OpenLogTooltip", "Opens a new tab in which you can filter log events related to this client"),
+		FSlateIcon(),
+		FUIAction(
+			FExecuteAction::CreateLambda([this]()
+			{
+				TArray<TSharedPtr<FClientBrowserItem>> Items = TileView->GetSelectedItems();
+				for (const TSharedPtr<FClientBrowserItem>& Item : Items)
+				{
+					OnListMouseButtonDoubleClick(Item);
+				}
+			}),
+			FCanExecuteAction::CreateLambda([] { return true; }),
+			FIsActionChecked()),
+		NAME_None,
+		EUserInterfaceActionType::Button
+		);
+	
+	return MenuBuilder.MakeWidget();
+}
+
+void UE::MultiUserServer::SConcertClientBrowser::OnListMouseButtonDoubleClick(TSharedPtr<FClientBrowserItem> ClientItem)
+{
+	OnClientDoubleClicked.ExecuteIfBound(ClientItem->ClientInfo.ClientEndpointId);
+}
+
+void UE::MultiUserServer::SConcertClientBrowser::SConcertClientBrowser::AllowAllSessions()
+{
+	AllowedSessions = BrowserModel->GetSessions();
+	UpdateTileViewFromAllowedSessions();
+}
+
+void UE::MultiUserServer::SConcertClientBrowser::DisallowAllSessions()
+{
+	AllowedSessions.Reset();
+	UpdateTileViewFromAllowedSessions();
+}
+
+void UE::MultiUserServer::SConcertClientBrowser::UpdateTileViewFromAllowedSessions()
+{
+	bShowAllSessions = AllowedSessions.Num() == BrowserModel->GetSessions().Num();
+	DisplayedClients.Empty();
+	
+	for (const FGuid& AllowedSession : AllowedSessions)
+	{
+		for (const FConcertSessionClientInfo& ClientInfo : BrowserModel->GetSessionClients(AllowedSession))
+		{
+			if (SessionFilter->PassesFilter(ClientInfo))
+			{
+				DisplayedClients.Add(MakeShared<FClientBrowserItem>(ClientInfo, BrowserModel->GetClientAddress(ClientInfo.ClientEndpointId)));
+			}
+		}
+	}
+	TileView->RequestListRefresh();
+}
+
+void UE::MultiUserServer::SConcertClientBrowser::GenerateSearchTerms(const FConcertSessionClientInfo& ClientInfo, TArray<FString>& SearchTerms) const
+{
+	if (const TSharedPtr<SConcertClientBrowserItem>* ItemWidget = ClientWidgets.Find(ClientInfo.ClientEndpointId))
+	{
+		ItemWidget->Get()->AppendSearchTerms(SearchTerms);
+	}
+}
+
+#undef LOCTEXT_NAMESPACE
diff --git a/Engine/Plugins/Developer/Concert/ConcertApp/MultiUserServer/Source/MultiUserServer/Private/Widgets/Clients/Browser/SConcertClientBrowser.h b/Engine/Plugins/Developer/Concert/ConcertApp/MultiUserServer/Source/MultiUserServer/Private/Widgets/Clients/Browser/SConcertClientBrowser.h
new file mode 100644
--- /dev/null
+++ b/Engine/Plugins/Developer/Concert/ConcertApp/MultiUserServer/Source/MultiUserServer/Private/Widgets/Clients/Browser/SConcertClientBrowser.h
@@ -0,0 +1,104 @@
+// Copyright Epic Games, Inc. All Rights Reserved.
+
+#pragma once
+
+#include "CoreMinimal.h"
+#include "SConcertClientBrowserItem.h"
+#include "Misc/TextFilter.h"
+#include "Widgets/DeclarativeSyntaxSupport.h"
+#include "Widgets/SCompoundWidget.h"
+
+class FConcertClientsTabController;
+class STableViewBase;
+class ITableRow;
+template<typename T> class STileView;
+
+enum class EConcertClientStatus : uint8;
+struct FConcertSessionClientInfo;
+
+namespace UE::MultiUserServer
+{
+	class IClientNetworkStatisticsModel;
+	class IClientBrowserModel;
+	struct FClientBrowserItem;
+	
+	/**
+	 * Presents a view of client thumbnails, which include important statistics.
+	 * Clients can be right-clicked, opening a context menu, and double clicked, opening a new tab with a log view.
+	 */
+	class SConcertClientBrowser : public SCompoundWidget
+	{
+	public:
+
+		DECLARE_DELEGATE_OneParam(FOnClientDoubleClicked, const FGuid& /*EndpointId*/);
+
+		SLATE_BEGIN_ARGS(SConcertClientBrowser) {}
+			/** Extension point to the right of the search bar */
+			SLATE_NAMED_SLOT(FArguments, RightOfSearch)
+			SLATE_EVENT(FOnClientDoubleClicked, OnClientDoubleClicked)
+		SLATE_END_ARGS()
+
+		void Construct(const FArguments& InArgs, TSharedRef<IClientBrowserModel> InBrowserModel, TSharedRef<IClientNetworkStatisticsModel> InStatisticsModel);
+
+		/** Shows only the clients connected to the given session ID */
+		void ShowOnlyClientsFromSession(const FGuid& SessionId);
+		
+	private:
+
+		using FSessionId = FGuid;
+		using FClientEndpointId = FGuid;
+		using FClientTextFilter = TTextFilter<const FConcertSessionClientInfo&>;
+
+		/** Retrieves clients and live sessions */
+		TSharedPtr<IClientBrowserModel> BrowserModel;
+		/** Used by sub-widgets to display networking data */
+		TSharedPtr<IClientNetworkStatisticsModel> StatisticsModel;
+
+		/** Contains the guid of every session we're allowed to display */
+		TSet<FSessionId> AllowedSessions;
+		/** Should all sessions be shown? */
+		bool bShowAllSessions = true;
+		/** Keeps widgets alive even when not displayed in the view - otherwise the graph will lose the historic data. */
+		TMap<FClientEndpointId, TSharedPtr<SConcertClientBrowserItem>> ClientWidgets;
+
+		/** Source array for TileView */
+		TArray<TSharedPtr<FClientBrowserItem>> DisplayedClients;
+		/** Visualizes all the items */
+		TSharedPtr<STileView<TSharedPtr<FClientBrowserItem>>> TileView;
+
+		// Filtering
+		TSharedPtr<FText> HighlightText;
+		TSharedPtr<FClientTextFilter> SessionFilter;
+		
+		FOnClientDoubleClicked OnClientDoubleClicked;
+
+		TSharedRef<SWidget> CreateSearchArea(const FArguments& InArgs);
+		TSharedRef<SWidget> CreateTileView();
+
+		// Model events
+		void OnSessionCreated(const FGuid& SessionId);
+		void OnSessionDestroyed(const FGuid& SessionId);
+		void OnClientListChanged(const FGuid& SessionId, EConcertClientStatus UpdateType, const FConcertSessionClientInfo& ClientInfo);
+		
+		void RemoveClient(const FConcertSessionClientInfo& ClientInfo);
+		void UpdateClientInfo(const FConcertSessionClientInfo& ClientInfo);
+
+		// Combo button
+		TSharedRef<SWidget> MakeSessionOption();
+
+		FText GetErrorMessageText() const;
+		
+		// TileView events
+		TSharedRef<ITableRow> MakeTileViewWidget(TSharedPtr<FClientBrowserItem> ClientItem, const TSharedRef<STableViewBase>& OwnerTable);
+		TSharedPtr<SWidget> OnGetContextMenuContent();
+		void OnListMouseButtonDoubleClick(TSharedPtr<FClientBrowserItem> ClientItem);
+
+		// Filtering
+		void AllowAllSessions();
+		void DisallowAllSessions();
+		void UpdateTileViewFromAllowedSessions();
+		void GenerateSearchTerms(const FConcertSessionClientInfo& ClientInfo, TArray<FString>& SearchTerms) const;
+	};
+}
+
+
diff --git a/Engine/Plugins/Developer/Concert/ConcertApp/MultiUserServer/Source/MultiUserServer/Private/Widgets/Clients/Browser/SConcertClientBrowserItem.cpp b/Engine/Plugins/Developer/Concert/ConcertApp/MultiUserServer/Source/MultiUserServer/Private/Widgets/Clients/Browser/SConcertClientBrowserItem.cpp
new file mode 100644
--- /dev/null
+++ b/Engine/Plugins/Developer/Concert/ConcertApp/MultiUserServer/Source/MultiUserServer/Private/Widgets/Clients/Browser/SConcertClientBrowserItem.cpp
@@ -0,0 +1,141 @@
+// Copyright Epic Games, Inc. All Rights Reserved.
+
+#include "SConcertClientBrowserItem.h"
+
+#include "ClientBrowserItem.h"
+#include "ConcertServerStyle.h"
+#include "INetworkMessagingExtension.h"
+#include "SClientNetworkStats.h"
+#include "Models/IClientNetworkStatisticsModel.h"
+#include "Widgets/SBoxPanel.h"
+#include "Widgets/Layout/SBorder.h"
+#include "Widgets/Layout/SScaleBox.h"
+#include "Widgets/Text/STextBlock.h"
+
+void UE::MultiUserServer::SConcertClientBrowserItem::Construct(const FArguments& InArgs, TSharedRef<FClientBrowserItem> InClientItem, TSharedRef<IClientNetworkStatisticsModel> InStatModel)
+{
+	Item = MoveTemp(InClientItem);
+	StatModel = MoveTemp(InStatModel);
+	HighlightText = InArgs._HighlightText;
+
+	ChildSlot
+	.Padding(FMargin(0.0f, 0.0f, 4.0f, 4.0f))
+	[
+		// Shadow behind thumbnail
+		SNew(SBorder)
+		.BorderImage(FConcertServerStyle::Get().GetBrush("Concert.Clients.DropShadow"))
+		.Padding(FMargin(0.0f, 0.0f, 5.0f, 5.0f))
+		[
+			// Change outside of thumbnail depending on hover state; lighter than inside
+			SNew(SBorder)
+			.BorderImage(this, &SConcertClientBrowserItem::GetBackgroundImage)
+			.Padding(2.f)
+			[
+				// Inside of thumbnail is darker
+				SNew(SBorder)
+				.BorderImage(FConcertServerStyle::Get().GetBrush("Concert.Clients.ThumbnailTitle"))
+				.Padding(4.f)
+				[
+					SNew(SVerticalBox)
+
+					+SVerticalBox::Slot()
+					.FillHeight(1.f)
+					.VAlign(VAlign_Top)
+					[
+						CreateHeader()	
+					]
+
+					+SVerticalBox::Slot()
+					.AutoHeight()
+					.Padding(0.f, 5.f, 0.f, 0.f)
+					.HAlign(HAlign_Fill)
+					.VAlign(VAlign_Bottom)
+					[
+						CreateStats()
+					]
+
+					+SVerticalBox::Slot()
+					.AutoHeight()
+					.Padding(-4.f, 5.f, -4.f, -4.f)
+					.VAlign(VAlign_Bottom)
+					[
+						CreateFooter()
+					]
+				]
+			]
+		]
+	];
+
+	OnClientInfoChanged();
+}
+
+void UE::MultiUserServer::SConcertClientBrowserItem::OnClientInfoChanged()
+{
+	ClientName->SetText(FText::FromString(Item->ClientInfo.ClientInfo.DisplayName));
+}
+
+void UE::MultiUserServer::SConcertClientBrowserItem::AppendSearchTerms(TArray<FString>& SearchTerms) const
+{
+	NetworkStats->AppendSearchTerms(SearchTerms);
+	SearchTerms.Add(ClientName->GetText().ToString());
+	SearchTerms.Add(ClientIP4->GetText().ToString());
+}
+
+TSharedRef<SWidget> UE::MultiUserServer::SConcertClientBrowserItem::CreateHeader()
+{
+	return SAssignNew(ClientName, STextBlock)
+		.Font(FConcertServerStyle::Get().GetFontStyle("Concert.Clients.ClientNameTileFont"))
+		.Text_Lambda([this](){ return FText::FromString(Item->ClientInfo.ClientInfo.DisplayName);} )
+		.HighlightText_Lambda([this](){ return *HighlightText; })
+		.ColorAndOpacity(FColor::White);
+}
+
+TSharedRef<SWidget> UE::MultiUserServer::SConcertClientBrowserItem::CreateStats()
+{
+	return SNew(SScaleBox)
+		.Stretch(EStretch::ScaleToFit)
+		.HAlign(HAlign_Fill)
+		[
+			SAssignNew(NetworkStats, SClientNetworkStats, Item->ClientAddress, StatModel.ToSharedRef())
+			.HighlightText(HighlightText)
+		];
+}
+
+TSharedRef<SWidget> UE::MultiUserServer::SConcertClientBrowserItem::CreateFooter()
+{
+	return SNew(SBorder)
+		.BorderImage(FConcertServerStyle::Get().GetBrush("Concert.Clients.ThumbnailFooter"))
+		[
+			SNew(SHorizontalBox)
+					
+			+SHorizontalBox::Slot()
+			.HAlign(HAlign_Right)
+			.Padding(2.f)
+			[
+				SAssignNew(ClientIP4, STextBlock)
+				.ColorAndOpacity(FColor::White)
+				.HighlightText_Lambda([this](){ return *HighlightText; })
+				.Text_Lambda([this]()
+				{
+					const FString Fallback(TEXT("No IP available"));
+					if (const TOptional<FMessageTransportStatistics> Statistics = StatModel->GetLatestNetworkStatistics(Item->ClientAddress))
+					{
+						const FString DisplayString = Statistics->IPv4AsString.IsEmpty()
+							? Fallback
+							: Statistics->IPv4AsString;
+						return FText::FromString(DisplayString);
+					}
+					return FText::FromString(Fallback);
+				})
+			]
+		];
+}
+
+const FSlateBrush* UE::MultiUserServer::SConcertClientBrowserItem::GetBackgroundImage() const
+{
+	if (IsHovered())
+	{
+		return FConcertServerStyle::Get().GetBrush("Concert.Clients.ThumbnailAreaHoverBackground");
+	}
+	return FConcertServerStyle::Get().GetBrush("Concert.Clients.ThumbnailAreaBackground");
+}
diff --git a/Engine/Plugins/Developer/Concert/ConcertApp/MultiUserServer/Source/MultiUserServer/Private/Widgets/Clients/Browser/SConcertClientBrowserItem.h b/Engine/Plugins/Developer/Concert/ConcertApp/MultiUserServer/Source/MultiUserServer/Private/Widgets/Clients/Browser/SConcertClientBrowserItem.h
new file mode 100644
--- /dev/null
+++ b/Engine/Plugins/Developer/Concert/ConcertApp/MultiUserServer/Source/MultiUserServer/Private/Widgets/Clients/Browser/SConcertClientBrowserItem.h
@@ -0,0 +1,53 @@
+// Copyright Epic Games, Inc. All Rights Reserved.
+
+#pragma once
+
+#include "CoreMinimal.h"
+#include "IMessageContext.h"
+#include "Widgets/DeclarativeSyntaxSupport.h"
+#include "Widgets/SCompoundWidget.h"
+
+class STextBlock;
+
+namespace UE::MultiUserServer
+{
+	class SClientNetworkStats;
+	class IClientNetworkStatisticsModel;
+	struct FClientBrowserItem;
+
+	class SConcertClientBrowserItem : public SCompoundWidget
+	{
+	public:
+
+		SLATE_BEGIN_ARGS(SConcertClientBrowserItem) {}
+			SLATE_ARGUMENT(TSharedPtr<FText>, HighlightText)
+		SLATE_END_ARGS()
+
+		void Construct(const FArguments& InArgs, TSharedRef<FClientBrowserItem> InClientItem, TSharedRef<IClientNetworkStatisticsModel> InStatModel);
+
+		/** Called when the underlying client data was changed */
+		void OnClientInfoChanged();
+
+		/** Appends the statistics to the search terms */
+		void AppendSearchTerms(TArray<FString>& SearchTerms) const;
+
+	private:
+
+		/** What we're displaying */
+		TSharedPtr<FClientBrowserItem> Item;
+		TSharedPtr<IClientNetworkStatisticsModel> StatModel;
+
+		/** The text to highlight */
+		TSharedPtr<FText> HighlightText;
+
+		TSharedPtr<STextBlock> ClientName;
+		TSharedPtr<SClientNetworkStats> NetworkStats;
+		TSharedPtr<STextBlock> ClientIP4;
+
+		TSharedRef<SWidget> CreateHeader();
+		TSharedRef<SWidget> CreateStats();
+		TSharedRef<SWidget> CreateFooter();
+
+		const FSlateBrush* GetBackgroundImage() const;
+	};
+}
diff --git a/Engine/Plugins/Developer/Concert/ConcertApp/MultiUserServer/Source/MultiUserServer/Private/Widgets/Clients/ConcertClientsTabController.cpp b/Engine/Plugins/Developer/Concert/ConcertApp/MultiUserServer/Source/MultiUserServer/Private/Widgets/Clients/ConcertClientsTabController.cpp
--- a/Engine/Plugins/Developer/Concert/ConcertApp/MultiUserServer/Source/MultiUserServer/Private/Widgets/Clients/ConcertClientsTabController.cpp
+++ b/Engine/Plugins/Developer/Concert/ConcertApp/MultiUserServer/Source/MultiUserServer/Private/Widgets/Clients/ConcertClientsTabController.cpp
@@ -3,12 +3,22 @@
 #include "ConcertClientsTabController.h"
 
 #include "SConcertClientsTabView.h"
+#include "Logging/Source/GlobalLogSource.h"
 #include "Window/ConcertServerTabs.h"
 #include "Window/ConcertServerWindowController.h"
 #include "Widgets/Docking/SDockTab.h"
 
 #define LOCTEXT_NAMESPACE "UnrealMultiUserUI"
 
+namespace UE::MultiUserServer
+{
+	static constexpr uint64 GlobalLogCapacity = 500000;
+}
+
+FConcertClientsTabController::FConcertClientsTabController()
+	: LogBuffer(MakeShared<FGlobalLogSource>(UE::MultiUserServer::GlobalLogCapacity))
+{}
+
 void FConcertClientsTabController::Init(const FConcertComponentInitParams& Params)
 {
 	FGlobalTabmanager::Get()->RegisterTabSpawner(
@@ -21,6 +31,12 @@ void FConcertClientsTabController::Init(const FConcertComponentInitParams& Param
 	Params.MainStack->AddTab(ConcertServerTabs::GetClientsTabID(), ETabState::OpenedTab);
 }
 
+void FConcertClientsTabController::ShowConnectedClients(const FGuid& SessionId) const
+{
+	FGlobalTabmanager::Get()->TryInvokeTab(ConcertServerTabs::GetClientsTabID());
+	ClientsView->ShowConnectedClients(SessionId);
+}
+
 TSharedRef<SDockTab> FConcertClientsTabController::SpawnClientsTab(const FSpawnTabArgs& Args, TSharedPtr<SWindow> RootWindow, TSharedRef<IConcertSyncServer> Server)
 {
 	const TSharedRef<SDockTab> DockTab = SNew(SDockTab)
@@ -28,7 +44,7 @@ TSharedRef<SDockTab> FConcertClientsTabController::SpawnClientsTab(const FSpawnT
 		.TabRole(MajorTab)
 		.CanEverClose(false);
 	DockTab->SetContent(
-		SNew(SConcertClientsTabView, ConcertServerTabs::GetClientsTabID(), Server)
+		SAssignNew(ClientsView, SConcertClientsTabView, ConcertServerTabs::GetClientsTabID(), Server, LogBuffer)
 			.ConstructUnderMajorTab(DockTab)
 			.ConstructUnderWindow(RootWindow)
 		);
diff --git a/Engine/Plugins/Developer/Concert/ConcertApp/MultiUserServer/Source/MultiUserServer/Private/Widgets/Clients/ConcertClientsTabController.h b/Engine/Plugins/Developer/Concert/ConcertApp/MultiUserServer/Source/MultiUserServer/Private/Widgets/Clients/ConcertClientsTabController.h
--- a/Engine/Plugins/Developer/Concert/ConcertApp/MultiUserServer/Source/MultiUserServer/Private/Widgets/Clients/ConcertClientsTabController.h
+++ b/Engine/Plugins/Developer/Concert/ConcertApp/MultiUserServer/Source/MultiUserServer/Private/Widgets/Clients/ConcertClientsTabController.h
@@ -5,17 +5,30 @@
 #include "CoreMinimal.h"
 #include "Window/IConcertComponent.h"
 
+class FGlobalLogSource;
+class SConcertClientsTabView;
 class SWindow;
 
-class FConcertClientsTabController : public IConcertComponent
+class FConcertClientsTabController : public IConcertComponent, public TSharedFromThis<FConcertClientsTabController>
 {
 public:
+	
+	FConcertClientsTabController();
 
 	//~ Begin IConcertComponent Interface
 	virtual void Init(const FConcertComponentInitParams& Params) override;
 	//~ End IConcertComponent Interface
 
+	/** Highlights this tab and sets the client filter such that all connected clients of the given session ID are shown.  */
+	void ShowConnectedClients(const FGuid& SessionId) const;
+
 private:
 
+	/** Buffers generated logs up to a limit (and overrides oldest logs when buffer is full) */
+	TSharedRef<FGlobalLogSource> LogBuffer;
+	
+	/** Manages the sub-tabs */
+	TSharedPtr<SConcertClientsTabView> ClientsView;
+	
 	TSharedRef<SDockTab> SpawnClientsTab(const FSpawnTabArgs& Args, TSharedPtr<SWindow> RootWindow, TSharedRef<IConcertSyncServer>);
 };
diff --git a/Engine/Plugins/Developer/Concert/ConcertApp/MultiUserServer/Source/MultiUserServer/Private/Widgets/Clients/Logging/Filter/ConcertFrontendLogFilter_Client.cpp b/Engine/Plugins/Developer/Concert/ConcertApp/MultiUserServer/Source/MultiUserServer/Private/Widgets/Clients/Logging/Filter/ConcertFrontendLogFilter_Client.cpp
new file mode 100644
--- /dev/null
+++ b/Engine/Plugins/Developer/Concert/ConcertApp/MultiUserServer/Source/MultiUserServer/Private/Widgets/Clients/Logging/Filter/ConcertFrontendLogFilter_Client.cpp
@@ -0,0 +1,20 @@
+// Copyright Epic Games, Inc. All Rights Reserved.
+
+#include "ConcertFrontendLogFilter_Client.h"
+
+#include "ConcertTransportEvents.h"
+
+bool FConcertLogFilter_Client::PassesFilter(const FConcertLog& InItem) const
+{
+	return AllowedClientEndpointIds.Contains(InItem.OriginEndpointId)
+		|| AllowedClientEndpointIds.Contains(InItem.DestinationEndpointId);
+}
+
+void FConcertLogFilter_Client::AllowOnly(const FGuid& ClientEndpointId)
+{
+	if (AllowedClientEndpointIds.Num() > 1 || !AllowedClientEndpointIds.Contains(ClientEndpointId))
+	{
+		AllowedClientEndpointIds = { ClientEndpointId };
+		OnChanged().Broadcast();
+	}
+}
diff --git a/Engine/Plugins/Developer/Concert/ConcertApp/MultiUserServer/Source/MultiUserServer/Private/Widgets/Clients/Logging/Filter/ConcertFrontendLogFilter_Client.h b/Engine/Plugins/Developer/Concert/ConcertApp/MultiUserServer/Source/MultiUserServer/Private/Widgets/Clients/Logging/Filter/ConcertFrontendLogFilter_Client.h
new file mode 100644
--- /dev/null
+++ b/Engine/Plugins/Developer/Concert/ConcertApp/MultiUserServer/Source/MultiUserServer/Private/Widgets/Clients/Logging/Filter/ConcertFrontendLogFilter_Client.h
@@ -0,0 +1,28 @@
+// Copyright Epic Games, Inc. All Rights Reserved.
+
+#pragma once
+
+#include "CoreMinimal.h"
+#include "Widgets/Clients/Logging/Filter/ConcertFrontendLogFilter.h"
+
+/** Only allows messages from the given clients */
+class FConcertLogFilter_Client : public FConcertLogFilter
+{
+public:
+
+	FConcertLogFilter_Client() = default;
+	explicit FConcertLogFilter_Client(const FGuid& SingleAllowedId)
+		: AllowedClientEndpointIds({ SingleAllowedId })
+	{}
+
+	//~ Begin FConcertLogFilter Interface
+	virtual bool PassesFilter(const FConcertLog& InItem) const override;
+	//~ End FConcertLogFilter Interface
+
+	void AllowOnly(const FGuid& ClientEndpointId);
+
+private:
+
+	/** Messages to and from the following client endpoint IDs are allowed */
+	TSet<FGuid> AllowedClientEndpointIds;
+};
diff --git a/Engine/Plugins/Developer/Concert/ConcertApp/MultiUserServer/Source/MultiUserServer/Private/Widgets/Clients/Logging/Filter/ConcertLogFilter_FrontendRoot.cpp b/Engine/Plugins/Developer/Concert/ConcertApp/MultiUserServer/Source/MultiUserServer/Private/Widgets/Clients/Logging/Filter/ConcertLogFilter_FrontendRoot.cpp
--- a/Engine/Plugins/Developer/Concert/ConcertApp/MultiUserServer/Source/MultiUserServer/Private/Widgets/Clients/Logging/Filter/ConcertLogFilter_FrontendRoot.cpp
+++ b/Engine/Plugins/Developer/Concert/ConcertApp/MultiUserServer/Source/MultiUserServer/Private/Widgets/Clients/Logging/Filter/ConcertLogFilter_FrontendRoot.cpp
@@ -3,6 +3,7 @@
 #include "ConcertLogFilter_FrontendRoot.h"
 
 #include "ConcertFrontendLogFilter.h"
+#include "ConcertFrontendLogFilter_Client.h"
 #include "ConcertFrontendLogFilter_MessageAction.h"
 #include "ConcertFrontendLogFilter_MessageType.h"
 #include "ConcertFrontendLogFilter_TextSearch.h"
@@ -12,14 +13,20 @@
 #include "Algo/AnyOf.h"
 #include "Widgets/SBoxPanel.h"
 
-FConcertLogFilter_FrontendRoot::FConcertLogFilter_FrontendRoot(TSharedRef<FConcertLogTokenizer> Tokenizer, TArray<TSharedRef<FConcertFrontendLogFilter>> InCustomFilters)
+FConcertLogFilter_FrontendRoot::FConcertLogFilter_FrontendRoot(TSharedRef<FConcertLogTokenizer> Tokenizer, TArray<TSharedRef<FConcertFrontendLogFilter>> InCustomFilters, const TArray<TSharedRef<FConcertLogFilter>>& NonVisualFilters)
 	: TextSearchFilter(MakeShared<FConcertFrontendLogFilter_TextSearch>(MoveTemp(Tokenizer)))
 	, CustomFilters(MoveTemp(InCustomFilters))
 	, AllFrontendFilters(CustomFilters)
 {
+	AllFrontendFilters.Reserve(CustomFilters.Num() + 1 + NonVisualFilters.Num());
+	
 	AllFrontendFilters.Add(TextSearchFilter);
+	for (const TSharedRef<FConcertLogFilter>& NonVisualFilter : NonVisualFilters)
+	{
+		AllFrontendFilters.Add(NonVisualFilter);
+	}
 	
-	for (const TSharedRef<FConcertFrontendLogFilter>& Filter : AllFrontendFilters)
+	for (const TSharedRef<FConcertLogFilter>& Filter : AllFrontendFilters)
 	{
 		Filter->OnChanged().AddLambda([this]()
 		{
@@ -51,7 +58,7 @@ bool FConcertLogFilter_FrontendRoot::PassesFilter(const FConcertLog& InItem) con
 {
 	return Algo::AllOf(
 		AllFrontendFilters,
-		[&InItem](const TSharedRef<FConcertFrontendLogFilter>& AndFilter){ return AndFilter->PassesFilter(InItem); }
+		[&InItem](const TSharedRef<FConcertLogFilter>& AndFilter){ return AndFilter->PassesFilter(InItem); }
 		);
 }
 
@@ -76,23 +83,39 @@ TSharedRef<SWidget> FConcertLogFilter_FrontendRoot::BuildCustomFilterListWidget(
 	return Box;
 }
 
-namespace UE::MultiUserServer::Private
+namespace UE::MultiUserServer
 {
-	static TArray<TSharedRef<FConcertFrontendLogFilter>> CreateCommonFilters()
+	namespace Private
 	{
-		return {
-			MakeShared<FConcertFrontendLogFilter_MessageAction>(),
-			MakeShared<FConcertFrontendLogFilter_MessageType>(),
-			MakeShared<FConcertFrontendLogFilter_Time>(ETimeFilter::AllowAfter),
-			MakeShared<FConcertFrontendLogFilter_Time>(ETimeFilter::AllowBefore)
-		};
+		static TArray<TSharedRef<FConcertFrontendLogFilter>> CreateCommonFilters()
+		{
+			return {
+				MakeShared<FConcertFrontendLogFilter_MessageAction>(),
+				MakeShared<FConcertFrontendLogFilter_MessageType>(),
+				MakeShared<FConcertFrontendLogFilter_Time>(ETimeFilter::AllowAfter),
+				MakeShared<FConcertFrontendLogFilter_Time>(ETimeFilter::AllowBefore)
+			};
+		}
+	}
+	
+	TSharedRef<FConcertLogFilter_FrontendRoot> MakeGlobalLogFilter(TSharedRef<FConcertLogTokenizer> Tokenizer)
+	{
+		return MakeShared<FConcertLogFilter_FrontendRoot>(
+			MoveTemp(Tokenizer),
+			Private::CreateCommonFilters()
+			);
 	}
-}
 
-TSharedRef<FConcertLogFilter_FrontendRoot> UE::MultiUserServer::MakeGlobalLogFilter(TSharedRef<FConcertLogTokenizer> Tokenizer)
-{
-	return MakeShared<FConcertLogFilter_FrontendRoot>(
-		MoveTemp(Tokenizer),
-		Private::CreateCommonFilters()
-		);
+	TSharedRef<FConcertLogFilter_FrontendRoot> MakeClientLogFilter(TSharedRef<FConcertLogTokenizer> Tokenizer, const FGuid& ClientEndpointId)
+	{
+		const TArray<TSharedRef<FConcertFrontendLogFilter>> CommonFilters = Private::CreateCommonFilters();
+		const TArray<TSharedRef<FConcertLogFilter>> NonVisuals = {
+			MakeShared<FConcertLogFilter_Client>(ClientEndpointId)
+		};
+		return MakeShared<FConcertLogFilter_FrontendRoot>(
+			MoveTemp(Tokenizer),
+			CommonFilters,
+			NonVisuals
+			);
+	}
 }
diff --git a/Engine/Plugins/Developer/Concert/ConcertApp/MultiUserServer/Source/MultiUserServer/Private/Widgets/Clients/Logging/Filter/ConcertLogFilter_FrontendRoot.h b/Engine/Plugins/Developer/Concert/ConcertApp/MultiUserServer/Source/MultiUserServer/Private/Widgets/Clients/Logging/Filter/ConcertLogFilter_FrontendRoot.h
--- a/Engine/Plugins/Developer/Concert/ConcertApp/MultiUserServer/Source/MultiUserServer/Private/Widgets/Clients/Logging/Filter/ConcertLogFilter_FrontendRoot.h
+++ b/Engine/Plugins/Developer/Concert/ConcertApp/MultiUserServer/Source/MultiUserServer/Private/Widgets/Clients/Logging/Filter/ConcertLogFilter_FrontendRoot.h
@@ -15,7 +15,9 @@ class FConcertLogFilter_FrontendRoot : public FConcertLogFilter, public TSharedF
 {
 public:
 
-	FConcertLogFilter_FrontendRoot(TSharedRef<FConcertLogTokenizer> Tokenizer, TArray<TSharedRef<FConcertFrontendLogFilter>> InCustomFilters);
+	FConcertLogFilter_FrontendRoot(TSharedRef<FConcertLogTokenizer> Tokenizer,
+		TArray<TSharedRef<FConcertFrontendLogFilter>> InCustomFilters,
+		const TArray<TSharedRef<FConcertLogFilter>>& NonVisualFilters = {});
 
 	/** Builds the widget view for all contained filters */
 	TSharedRef<SWidget> BuildFilterWidgets() const;
@@ -25,7 +27,6 @@ public:
 	//~ End IFilter Interface
 
 	FORCEINLINE const TSharedRef<FConcertFrontendLogFilter_TextSearch>& GetTextSearchFilter() const { return TextSearchFilter; }
-	FORCEINLINE const TArray<TSharedRef<FConcertFrontendLogFilter>>& GetFrontendFilters() const { return AllFrontendFilters; }
 
 private:
 
@@ -35,8 +36,8 @@ private:
 	/** AllFrontendFilters without special filters we have as propteries above, such as TextSearchFilter. */
 	TArray<TSharedRef<FConcertFrontendLogFilter>> CustomFilters;
 	
-	/** Filters that are combined using logical OR. */
-	TArray<TSharedRef<FConcertFrontendLogFilter>> AllFrontendFilters;
+	/** Filters that are combined using logical AND. */
+	TArray<TSharedRef<FConcertLogFilter>> AllFrontendFilters;
 
 	/** Builds the widgets that go under the text */
 	TSharedRef<SWidget> BuildCustomFilterListWidget() const;
@@ -46,4 +47,7 @@ namespace UE::MultiUserServer
 {
 	/** Creates a filter for the global filter log window */
 	TSharedRef<FConcertLogFilter_FrontendRoot> MakeGlobalLogFilter(TSharedRef<FConcertLogTokenizer> Tokenizer);
+
+	/** Creates a filter for a filter log window intended for a client */
+	TSharedRef<FConcertLogFilter_FrontendRoot> MakeClientLogFilter(TSharedRef<FConcertLogTokenizer> Tokenizer, const FGuid& ClientEndpointId);
 }
diff --git a/Engine/Plugins/Developer/Concert/ConcertApp/MultiUserServer/Source/MultiUserServer/Private/Widgets/Clients/Logging/SConcertTransportLog.cpp b/Engine/Plugins/Developer/Concert/ConcertApp/MultiUserServer/Source/MultiUserServer/Private/Widgets/Clients/Logging/SConcertTransportLog.cpp
--- a/Engine/Plugins/Developer/Concert/ConcertApp/MultiUserServer/Source/MultiUserServer/Private/Widgets/Clients/Logging/SConcertTransportLog.cpp
+++ b/Engine/Plugins/Developer/Concert/ConcertApp/MultiUserServer/Source/MultiUserServer/Private/Widgets/Clients/Logging/SConcertTransportLog.cpp
@@ -10,6 +10,7 @@
 #include "SConcertTransportLogRow.h"
 #include "Settings/ConcertTransportLogSettings.h"
 #include "Settings/MultiUserServerColumnVisibilitySettings.h"
+#include "Widgets/Clients/SPromptConcertLoggingEnabled.h"
 
 #include "Framework/MultiBox/MultiBoxBuilder.h"
 #include "Styling/AppStyle.h"
@@ -23,6 +24,11 @@
 
 const FName SConcertTransportLog::FirstColumnId("AvatarColourColumnId");
 
+SConcertTransportLog::~SConcertTransportLog()
+{
+	ConcertTransportEvents::OnConcertTransportLoggingEnabledChangedEvent().RemoveAll(this);
+}
+
 void SConcertTransportLog::Construct(const FArguments& InArgs, TSharedRef<IConcertLogSource> LogSource)
 {
 	PagedLogList = MakeShared<FPagedFilteredConcertLogList>(MoveTemp(LogSource), InArgs._Filter);
@@ -30,35 +36,40 @@ void SConcertTransportLog::Construct(const FArguments& InArgs, TSharedRef<IConce
 	HighlightText = MakeShared<FText>();
 
 	GetClientInfoFunc = InArgs._GetClientInfo;
-
+	
 	ChildSlot
 	[
-		SNew(SBorder)
-		.BorderImage(FAppStyle::Get().GetBrush("ToolPanel.GroupBorder"))
-		.BorderBackgroundColor(FSlateColor(FLinearColor(0.6, 0.6, 0.6)))
-		.Padding(2)
-		[
-			SNew(SVerticalBox)
+		SAssignNew(EnableLoggingPromptOverlay, SOverlay)
 
-			+SVerticalBox::Slot()
-			.AutoHeight()
-			.VAlign(VAlign_Top)
+		+SOverlay::Slot()
+		[
+			SNew(SBorder)
+			.BorderImage(FAppStyle::Get().GetBrush("ToolPanel.GroupBorder"))
+			.BorderBackgroundColor(FSlateColor(FLinearColor(0.6, 0.6, 0.6)))
+			.Padding(2)
 			[
-				InArgs._Filter ? InArgs._Filter->BuildFilterWidgets() : SNullWidget::NullWidget
-			]
+				SNew(SVerticalBox)
 
-			+SVerticalBox::Slot()
-			.FillHeight(1.f)
-			.Padding(0, 5, 0, 0)
-			[
-				CreateTableView()
-			]
+				+SVerticalBox::Slot()
+				.AutoHeight()
+				.VAlign(VAlign_Top)
+				[
+					InArgs._Filter ? InArgs._Filter->BuildFilterWidgets() : SNullWidget::NullWidget
+				]
 
-			+SVerticalBox::Slot()
-			.AutoHeight()
-			[
-				SNew(SConcertTransportLogFooter, PagedLogList.ToSharedRef())
-				.ExtendViewOptions(this, &SConcertTransportLog::ExtendViewOptions)
+				+SVerticalBox::Slot()
+				.FillHeight(1.f)
+				.Padding(0, 5, 0, 0)
+				[
+					CreateTableView()
+				]
+
+				+SVerticalBox::Slot()
+				.AutoHeight()
+				[
+					SNew(SConcertTransportLogFooter, PagedLogList.ToSharedRef())
+					.ExtendViewOptions(this, &SConcertTransportLog::ExtendViewOptions)
+				]
 			]
 		]
 	];
@@ -69,9 +80,12 @@ void SConcertTransportLog::Construct(const FArguments& InArgs, TSharedRef<IConce
 	{
 		InArgs._Filter->GetTextSearchFilter()->OnSearchTextChanged().AddSP(this, &SConcertTransportLog::OnSearchTextChanged);
 	}
-	UMultiUserServerColumnVisibilitySettings::GetSettings()->OnTransportLogColumnVisibility().AddSP(this, &SConcertTransportLog::OnColumnVisibilitySettingsChanged);
 	
+	UMultiUserServerColumnVisibilitySettings::GetSettings()->OnTransportLogColumnVisibility().AddSP(this, &SConcertTransportLog::OnColumnVisibilitySettingsChanged);
 	UE::ConcertSharedSlate::RestoreColumnVisibilityState(HeaderRow.ToSharedRef(), UMultiUserServerColumnVisibilitySettings::GetSettings()->GetTransportLogColumnVisibility());
+	
+	ConcertTransportEvents::OnConcertTransportLoggingEnabledChangedEvent().AddSP(this, &SConcertTransportLog::OnConcertLoggingEnabledChanged);
+	OnConcertLoggingEnabledChanged(ConcertTransportEvents::IsLoggingEnabled());
 }
 
 TSharedRef<SWidget> SConcertTransportLog::CreateTableView()
@@ -228,4 +242,16 @@ void SConcertTransportLog::OnColumnVisibilitySettingsChanged(const FColumnVisibi
 	UE::ConcertSharedSlate::RestoreColumnVisibilityState(HeaderRow.ToSharedRef(), ColumnSnapshot);
 }
 
-#undef LOCTEXT_NAMESPACE
\ No newline at end of file
+void SConcertTransportLog::OnConcertLoggingEnabledChanged(bool bNewEnabled)
+{
+	if (!bNewEnabled)
+	{
+		EnableLoggingPromptOverlay->AddSlot().AttachWidget(SAssignNew(EnableLoggingPrompt, SPromptConcertLoggingEnabled));
+	}
+	else if (EnableLoggingPrompt)
+	{
+		EnableLoggingPromptOverlay->RemoveSlot(EnableLoggingPrompt.ToSharedRef());
+	}
+}
+
+#undef LOCTEXT_NAMESPACE
diff --git a/Engine/Plugins/Developer/Concert/ConcertApp/MultiUserServer/Source/MultiUserServer/Private/Widgets/Clients/Logging/SConcertTransportLog.h b/Engine/Plugins/Developer/Concert/ConcertApp/MultiUserServer/Source/MultiUserServer/Private/Widgets/Clients/Logging/SConcertTransportLog.h
--- a/Engine/Plugins/Developer/Concert/ConcertApp/MultiUserServer/Source/MultiUserServer/Private/Widgets/Clients/Logging/SConcertTransportLog.h
+++ b/Engine/Plugins/Developer/Concert/ConcertApp/MultiUserServer/Source/MultiUserServer/Private/Widgets/Clients/Logging/SConcertTransportLog.h
@@ -8,7 +8,8 @@
 #include "Widgets/DeclarativeSyntaxSupport.h"
 #include "Widgets/SCompoundWidget.h"
 
-struct FColumnVisibilitySnapshot;
+class SOverlay;
+class SPromptConcertLoggingEnabled;
 class FConcertLogTokenizer;
 class FConcertLogFilter_FrontendRoot;
 class FMenuBuilder;
@@ -19,6 +20,7 @@ class SHeaderRow;
 class STableViewBase;
 template <typename ItemType> class SListView;
 
+struct FColumnVisibilitySnapshot;
 struct FConcertLogEntry;
 
 /**
@@ -37,11 +39,17 @@ public:
 		/** Used to better display client info, e.g. client name and colour */
 		SLATE_EVENT(FGetClientInfo, GetClientInfo)
 	SLATE_END_ARGS()
+	virtual ~SConcertTransportLog() override;
 
 	void Construct(const FArguments& InArgs, TSharedRef<IConcertLogSource> LogSource);
 
 private:
-
+	
+	/** Used to overlay EnableLoggingPrompt over the tabs */
+	TSharedPtr<SOverlay> EnableLoggingPromptOverlay;
+	/** Reminds the user to enable logging */
+	TSharedPtr<SPromptConcertLoggingEnabled> EnableLoggingPrompt;
+	
 	/** Sorts the log into pages whilst applying filters */
 	TSharedPtr<FPagedFilteredConcertLogList> PagedLogList;
 	/** Used by various systems to convert logs to text */
@@ -72,5 +80,7 @@ private:
 	void OnPageViewChanged(const TArray<TSharedPtr<FConcertLogEntry>>&);
 	void OnSearchTextChanged(const FText& NewSearchText);
 	void OnColumnVisibilitySettingsChanged(const FColumnVisibilitySnapshot& ColumnSnapshot);
+	
+	void OnConcertLoggingEnabledChanged(bool bNewEnabled);
 };
 
diff --git a/Engine/Plugins/Developer/Concert/ConcertApp/MultiUserServer/Source/MultiUserServer/Private/Widgets/Clients/SConcertClientsTabView.cpp b/Engine/Plugins/Developer/Concert/ConcertApp/MultiUserServer/Source/MultiUserServer/Private/Widgets/Clients/SConcertClientsTabView.cpp
--- a/Engine/Plugins/Developer/Concert/ConcertApp/MultiUserServer/Source/MultiUserServer/Private/Widgets/Clients/SConcertClientsTabView.cpp
+++ b/Engine/Plugins/Developer/Concert/ConcertApp/MultiUserServer/Source/MultiUserServer/Private/Widgets/Clients/SConcertClientsTabView.cpp
@@ -2,31 +2,34 @@
 
 #include "SConcertClientsTabView.h"
 
+#include "ConcertClientsTabController.h"
+#include "ConcertUtil.h"
 #include "IConcertServer.h"
 #include "IConcertSyncServer.h"
 #include "Logging/Filter/ConcertLogFilter_FrontendRoot.h"
 #include "Logging/Source/GlobalLogSource.h"
 #include "Logging/Util/ConcertLogTokenizer.h"
-#include "SPromptConcertLoggingEnabled.h"
+#include "Widgets/Clients/Browser/SConcertClientBrowser.h"
+#include "Widgets/Clients/Browser/Models/ClientBrowserModel.h"
+#include "Widgets/Clients/Browser/Models/ClientNetworkStatisticsModel.h"
 #include "Widgets/Clients/Logging/SConcertTransportLog.h"
 
 #include "Framework/Docking/TabManager.h"
 #include "Widgets/Docking/SDockTab.h"
+#include "Widgets/Images/SImage.h"
+#include "Widgets/Input/SButton.h"
+#include "Widgets/Text/STextBlock.h"
 
 #define LOCTEXT_NAMESPACE "UnrealMultiUserUI"
 
 const FName SConcertClientsTabView::ClientBrowserTabId("ClientBrowserTabId");
 const FName SConcertClientsTabView::GlobalLogTabId("GlobalLogTabId");
 
-SConcertClientsTabView::~SConcertClientsTabView()
+void SConcertClientsTabView::Construct(const FArguments& InArgs, FName InStatusBarID, TSharedRef<IConcertSyncServer> InServer, TSharedRef<FGlobalLogSource> InLogBuffer)
 {
-	ConcertTransportEvents::OnConcertTransportLoggingEnabledChangedEvent().RemoveAll(this);
-}
-
-void SConcertClientsTabView::Construct(const FArguments& InArgs, FName InStatusBarID, TSharedRef<IConcertSyncServer> InServer)
-{
-	LogTokenizer = MakeShared<FConcertLogTokenizer>();
 	Server = MoveTemp(InServer);
+	LogBuffer = MoveTemp(InLogBuffer);
+	LogTokenizer = MakeShared<FConcertLogTokenizer>();
 	
 	SConcertTabViewWithManagerBase::Construct(
 		SConcertTabViewWithManagerBase::FArguments()
@@ -36,12 +39,76 @@ void SConcertClientsTabView::Construct(const FArguments& InArgs, FName InStatusB
 		{
 			CreateTabs(InTabManager, InLayout, InArgs);
 		}))
-		.OverlayTabs(this, &SConcertClientsTabView::SetupLoggingPromptOverlay)
 		.LayoutName("ConcertClientsTabView_v0.1"),
 		InStatusBarID
 	);
 }
 
+void SConcertClientsTabView::ShowConnectedClients(const FGuid& SessionId) const
+{
+	ClientBrowser->ShowOnlyClientsFromSession(SessionId);
+}
+
+void SConcertClientsTabView::OpenClientLogTab(const FGuid& ClientEndpointId) const
+{
+	const FName TabId = *ClientEndpointId.ToString();
+	if (const TSharedPtr<SDockTab> ExistingTab = GetTabManager()->FindExistingLiveTab(FTabId(TabId)))
+	{
+		GetTabManager()->DrawAttention(ExistingTab.ToSharedRef());
+	}
+	else if (const TOptional<FConcertSessionClientInfo> ClientInfo = ConcertUtil::GetConnectedClientInfo(*Server->GetConcertServer(), ClientEndpointId)
+		; ensure(ClientInfo))
+	{
+		const TSharedRef<SDockTab> NewTab = SNew(SDockTab)
+			.Label(FText::Format(LOCTEXT("ClientTabFmt", "{0} Log"), FText::FromString(ClientInfo->ClientInfo.DisplayName)))
+			.ToolTipText(FText::Format(LOCTEXT("ClientTabTooltipFmt", "Logs all networked requests originating or going to  client {0} (EndpointId = {1})"), FText::FromString(ClientInfo->ClientInfo.DisplayName), FText::FromString(ClientEndpointId.ToString())))
+			.TabRole(PanelTab)
+			[
+				SNew(SConcertTransportLog, LogBuffer.ToSharedRef())
+				.Filter(UE::MultiUserServer::MakeClientLogFilter(LogTokenizer.ToSharedRef(), ClientEndpointId))
+			]; 
+
+		// We need a tab to place the client tab next to
+		if (IsGlobalLogOpen())
+		{
+			const FTabManager::FLiveTabSearch Search(GlobalLogTabId);
+			GetTabManager()->InsertNewDocumentTab(TabId, Search, NewTab);
+		}
+		else
+		{
+			OpenGlobalLogTab();
+			
+			const FTabManager::FLiveTabSearch Search(GlobalLogTabId);
+			GetTabManager()->InsertNewDocumentTab(TabId, Search, NewTab);
+
+			CloseGlobalLogTab();
+		}
+	}
+}
+
+void SConcertClientsTabView::OpenGlobalLogTab() const
+{
+	GetTabManager()->TryInvokeTab(GlobalLogTabId);
+}
+
+void SConcertClientsTabView::CloseGlobalLogTab() const
+{
+	if (const TSharedPtr<SDockTab> GlobalLogTab = GetGlobalLogTab())
+	{
+		GlobalLogTab->RequestCloseTab();
+	}
+}
+
+bool SConcertClientsTabView::IsGlobalLogOpen() const
+{
+	return GetGlobalLogTab().IsValid();
+}
+
+TSharedPtr<SDockTab> SConcertClientsTabView::GetGlobalLogTab() const
+{
+	return GetTabManager()->FindExistingLiveTab(GlobalLogTabId);
+}
+
 void SConcertClientsTabView::CreateTabs(const TSharedRef<FTabManager>& InTabManager, const TSharedRef<FTabManager::FLayout>& InLayout, const FArguments& InArgs)
 {
 	InTabManager->RegisterTabSpawner(ClientBrowserTabId, FOnSpawnTab::CreateSP(this, &SConcertClientsTabView::SpawnClientBrowserTab))
@@ -73,50 +140,65 @@ TSharedRef<SDockTab> SConcertClientsTabView::SpawnClientBrowserTab(const FSpawnT
 		.Label(LOCTEXT("ClientBrowserTabLabel", "Clients"))
 		.TabRole(PanelTab)
 		[
-			SNullWidget::NullWidget
+			SAssignNew(ClientBrowser, UE::MultiUserServer::SConcertClientBrowser,
+				MakeShared<UE::MultiUserServer::FClientBrowserModel>(Server->GetConcertServer()),
+				MakeShared<UE::MultiUserServer::FClientNetworkStatisticsModel>())
+			.RightOfSearch()
+			[
+				CreateOpenGlobalLogButton()
+			]
+			.OnClientDoubleClicked_Raw(this, &SConcertClientsTabView::OpenClientLogTab)
 		]; 
 }
 
 TSharedRef<SDockTab> SConcertClientsTabView::SpawnGlobalLogTab(const FSpawnTabArgs& InTabArgs)
 {
-	// TODO: Load from config file
-	constexpr size_t LogCapacity = 500000;
-	
 	return SNew(SDockTab)
 		.Label(LOCTEXT("GlobalLogTabLabel", "Global Log"))
 		.TabRole(PanelTab)
 		[
-			SNew(SConcertTransportLog, MakeShared<FGlobalLogSource>(LogCapacity))
+			SNew(SConcertTransportLog, LogBuffer.ToSharedRef())
 			.Filter(UE::MultiUserServer::MakeGlobalLogFilter(LogTokenizer.ToSharedRef()))
 		]; 
 }
 
-TSharedRef<SWidget> SConcertClientsTabView::SetupLoggingPromptOverlay(const TSharedRef<SWidget>& TabsWidget)
+TSharedRef<SWidget> SConcertClientsTabView::CreateOpenGlobalLogButton() const
 {
-	SOverlay::FOverlaySlot* MessageSlot;
-	EnableLoggingPromptOverlay =  SNew(SOverlay)
-		+SOverlay::Slot()
+	return SNew(SButton)
+		.ButtonStyle(FAppStyle::Get(), "SimpleButton")
+		.ToolTipText(LOCTEXT("OpenGlobalLogTooltip", "Opens the Global Log which logs all incoming networked messages."))
+		.ContentPadding(FMargin(1, 0))
+		.Visibility_Lambda(
+			[this]()
+			{
+				const bool bIsGlobalLogOpen = IsGlobalLogOpen();
+				return bIsGlobalLogOpen ? EVisibility::Hidden : EVisibility::Visible;
+			})
+		.OnClicked_Lambda([this]()
+		{
+			OpenGlobalLogTab();
+			return FReply::Handled();
+		})
 		[
-			TabsWidget
-		]
-		+SOverlay::Slot()
-		.Expose(MessageSlot);
-
-	MessageSlot->AttachWidget(SAssignNew(EnableLoggingPrompt, SPromptConcertLoggingEnabled));
-	ConcertTransportEvents::OnConcertTransportLoggingEnabledChangedEvent().AddSP(this, &SConcertClientsTabView::OnConcertLoggingEnabledChanged);
-	return EnableLoggingPromptOverlay.ToSharedRef();
-}
-
-void SConcertClientsTabView::OnConcertLoggingEnabledChanged(bool bNewEnabled)
-{
-	if (bNewEnabled)
-	{
-		EnableLoggingPromptOverlay->RemoveSlot(EnableLoggingPrompt.ToSharedRef());
-	}
-	else
-	{
-		EnableLoggingPromptOverlay->AddSlot().AttachWidget(SAssignNew(EnableLoggingPrompt, SPromptConcertLoggingEnabled));
-	}
+			SNew(SHorizontalBox)
+				+SHorizontalBox::Slot()
+				.AutoWidth()
+				.VAlign(VAlign_Center)
+				.Padding(4.0, 0.0f)
+				[
+					SNew(SImage)
+					.ColorAndOpacity(FSlateColor::UseForeground())
+					.Image(FAppStyle::Get().GetBrush("Icons.Layout"))
+				]
+				+SHorizontalBox::Slot()
+				.VAlign(VAlign_Center)
+				.Padding(4.0, 0.0f)
+				[
+					SNew(STextBlock)
+					.Text(LOCTEXT("OpenGlobalLog", "Open Global Log"))
+					.ColorAndOpacity(FSlateColor::UseForeground())
+				]
+		];
 }
 
 #undef LOCTEXT_NAMESPACE
diff --git a/Engine/Plugins/Developer/Concert/ConcertApp/MultiUserServer/Source/MultiUserServer/Private/Widgets/Clients/SConcertClientsTabView.h b/Engine/Plugins/Developer/Concert/ConcertApp/MultiUserServer/Source/MultiUserServer/Private/Widgets/Clients/SConcertClientsTabView.h
--- a/Engine/Plugins/Developer/Concert/ConcertApp/MultiUserServer/Source/MultiUserServer/Private/Widgets/Clients/SConcertClientsTabView.h
+++ b/Engine/Plugins/Developer/Concert/ConcertApp/MultiUserServer/Source/MultiUserServer/Private/Widgets/Clients/SConcertClientsTabView.h
@@ -6,14 +6,24 @@
 #include "Widgets/DeclarativeSyntaxSupport.h"
 #include "Widgets/SConcertTabViewWithManagerBase.h"
 
-struct FConcertClientInfo;
-class IConcertSyncServer;
+namespace UE
+{
+	namespace MultiUserServer
+	{
+		class SConcertClientBrowser;
+	}
+}
+
+class FConcertClientsTabController;
 class FConcertLogTokenizer;
 class FTabManager;
+class FGlobalLogSource;
+class IConcertSyncServer;
 class SDockTab;
 class SPromptConcertLoggingEnabled;
 class SWidget;
 class SWindow;
+struct FConcertClientInfo;
 
 /** Manages the UI logic of the Clients tab */
 class SConcertClientsTabView : public SConcertTabViewWithManagerBase
@@ -28,27 +38,36 @@ public:
 		SLATE_ARGUMENT(TSharedPtr<SDockTab>, ConstructUnderMajorTab)
 		SLATE_ARGUMENT(TSharedPtr<SWindow>, ConstructUnderWindow)
 	SLATE_END_ARGS()
-	virtual ~SConcertClientsTabView() override;
-
-	void Construct(const FArguments& InArgs, FName InStatusBarID, TSharedRef<IConcertSyncServer> InServer);
 
+	void Construct(const FArguments& InArgs, FName InStatusBarID, TSharedRef<IConcertSyncServer> InServer, TSharedRef<FGlobalLogSource> InLogBuffer);
+	
+	void ShowConnectedClients(const FGuid& SessionId) const;
+	void OpenGlobalLogTab() const;
+	void CloseGlobalLogTab() const;
+	void OpenClientLogTab(const FGuid& ClientEndpointId) const;
+	
+	bool IsGlobalLogOpen() const;
+	TSharedPtr<SDockTab> GetGlobalLogTab() const;
+	
 private:
 
+	/** Used to look up client info */
+	TSharedPtr<IConcertSyncServer> Server;
+	/** Buffers all logs globally */
+	TSharedPtr<FGlobalLogSource> LogBuffer;
+	/** Used by various systems to convert logs to text */
+	TSharedPtr<FConcertLogTokenizer> LogTokenizer;
+
+	TSharedPtr<UE::MultiUserServer::SConcertClientBrowser> ClientBrowser;
+	
 	/** Used to overlay EnableLoggingPrompt over the tabs */
 	TSharedPtr<SOverlay> EnableLoggingPromptOverlay;
 	/** Reminds the user to enable logging */
 	TSharedPtr<SPromptConcertLoggingEnabled> EnableLoggingPrompt;
-
-	/** Used by various systems to convert logs to text */
-	TSharedPtr<FConcertLogTokenizer> LogTokenizer;
-
-	/** Used to look up client info */
-	TSharedPtr<IConcertSyncServer> Server;
 	
 	void CreateTabs(const TSharedRef<FTabManager>& InTabManager, const TSharedRef<FTabManager::FLayout>& InLayout, const FArguments& InArgs);
 	TSharedRef<SDockTab> SpawnClientBrowserTab(const FSpawnTabArgs& InTabArgs);
 	TSharedRef<SDockTab> SpawnGlobalLogTab(const FSpawnTabArgs& InTabArgs);
 
-	TSharedRef<SWidget> SetupLoggingPromptOverlay(const TSharedRef<SWidget>& TabsWidget);
-	void OnConcertLoggingEnabledChanged(bool bNewEnabled);
+	TSharedRef<SWidget> CreateOpenGlobalLogButton() const;
 };
diff --git a/Engine/Plugins/Developer/Concert/ConcertApp/MultiUserServer/Source/MultiUserServer/Private/Widgets/Clients/SPromptConcertLoggingEnabled.cpp b/Engine/Plugins/Developer/Concert/ConcertApp/MultiUserServer/Source/MultiUserServer/Private/Widgets/Clients/SPromptConcertLoggingEnabled.cpp
--- a/Engine/Plugins/Developer/Concert/ConcertApp/MultiUserServer/Source/MultiUserServer/Private/Widgets/Clients/SPromptConcertLoggingEnabled.cpp
+++ b/Engine/Plugins/Developer/Concert/ConcertApp/MultiUserServer/Source/MultiUserServer/Private/Widgets/Clients/SPromptConcertLoggingEnabled.cpp
@@ -22,7 +22,7 @@ void SPromptConcertLoggingEnabled::Construct(const FArguments& InArgs)
 		.Padding(0)
 		[
 			SNew(SBox)
-			.Padding(0.f, 0.f, 0.f, 50.f)
+			.Padding(0.f, 0.f, 0.f, 10.f)
 			.HAlign(HAlign_Center)
 			.VAlign(VAlign_Center)
 			.Clipping(EWidgetClipping::ClipToBounds)
diff --git a/Engine/Plugins/Developer/Concert/ConcertApp/MultiUserServer/Source/MultiUserServer/Private/Widgets/SConcertTabViewWithManagerBase.h b/Engine/Plugins/Developer/Concert/ConcertApp/MultiUserServer/Source/MultiUserServer/Private/Widgets/SConcertTabViewWithManagerBase.h
--- a/Engine/Plugins/Developer/Concert/ConcertApp/MultiUserServer/Source/MultiUserServer/Private/Widgets/SConcertTabViewWithManagerBase.h
+++ b/Engine/Plugins/Developer/Concert/ConcertApp/MultiUserServer/Source/MultiUserServer/Private/Widgets/SConcertTabViewWithManagerBase.h
@@ -34,6 +34,10 @@ public:
 	 */
 	void Construct(const FArguments& InArgs, FName InStatusBarId);
 
+protected:
+
+	TSharedPtr<FTabManager> GetTabManager() const { return TabManager; }
+	
 private:
 
 	/** Holds the tab manager that manages the front-end's tabs. */
diff --git a/Engine/Plugins/Developer/Concert/ConcertApp/MultiUserServer/Source/MultiUserServer/Private/Widgets/SessionTabs/Live/LiveConcertSessionTab.cpp b/Engine/Plugins/Developer/Concert/ConcertApp/MultiUserServer/Source/MultiUserServer/Private/Widgets/SessionTabs/Live/LiveConcertSessionTab.cpp
--- a/Engine/Plugins/Developer/Concert/ConcertApp/MultiUserServer/Source/MultiUserServer/Private/Widgets/SessionTabs/Live/LiveConcertSessionTab.cpp
+++ b/Engine/Plugins/Developer/Concert/ConcertApp/MultiUserServer/Source/MultiUserServer/Private/Widgets/SessionTabs/Live/LiveConcertSessionTab.cpp
@@ -10,10 +10,11 @@
 
 #include "Widgets/Docking/SDockTab.h"
 
-FLiveConcertSessionTab::FLiveConcertSessionTab(TSharedRef<IConcertServerSession> InspectedSession, TSharedRef<IConcertSyncServer> SyncServer, TAttribute<TSharedRef<SWindow>> ConstructUnderWindow)
+FLiveConcertSessionTab::FLiveConcertSessionTab(TSharedRef<IConcertServerSession> InspectedSession, TSharedRef<IConcertSyncServer> SyncServer, TAttribute<TSharedRef<SWindow>> ConstructUnderWindow, FShowConnectedClients OnConnectedClientsClicked)
 	: FConcertSessionTabBase(InspectedSession->GetSessionInfo().SessionId, SyncServer)
 	, InspectedSession(MoveTemp(InspectedSession))
 	, ConstructUnderWindow(MoveTemp(ConstructUnderWindow))
+	, OnConnectedClientsClicked(MoveTemp(OnConnectedClientsClicked))
 	, SessionHistoryController(MakeShared<FLiveServerSessionHistoryController>(InspectedSession, SyncServer))
 	, PackageViewerController(MakeShared<FConcertSessionPackageViewerController>(InspectedSession, SyncServer))
 {}
@@ -29,7 +30,10 @@ void FLiveConcertSessionTab::CreateDockContent(const TSharedRef<SDockTab>& InDoc
 	};
 	InDockTab->SetContent(
 		SNew(SConcertLiveSessionTabView, WidgetArgs, *GetTabId())
-		);
+		.OnConnectedClientsClicked_Lambda([this]()
+		{
+			OnConnectedClientsClicked.ExecuteIfBound(InspectedSession);
+		}));
 }
 
 void FLiveConcertSessionTab::OnOpenTab()
diff --git a/Engine/Plugins/Developer/Concert/ConcertApp/MultiUserServer/Source/MultiUserServer/Private/Widgets/SessionTabs/Live/LiveConcertSessionTab.h b/Engine/Plugins/Developer/Concert/ConcertApp/MultiUserServer/Source/MultiUserServer/Private/Widgets/SessionTabs/Live/LiveConcertSessionTab.h
--- a/Engine/Plugins/Developer/Concert/ConcertApp/MultiUserServer/Source/MultiUserServer/Private/Widgets/SessionTabs/Live/LiveConcertSessionTab.h
+++ b/Engine/Plugins/Developer/Concert/ConcertApp/MultiUserServer/Source/MultiUserServer/Private/Widgets/SessionTabs/Live/LiveConcertSessionTab.h
@@ -21,7 +21,13 @@ class FLiveConcertSessionTab : public FConcertSessionTabBase
 {
 public:
 
-	FLiveConcertSessionTab(TSharedRef<IConcertServerSession> InspectedSession, TSharedRef<IConcertSyncServer> SyncServer, TAttribute<TSharedRef<SWindow>> ConstructUnderWindow);
+	DECLARE_DELEGATE_OneParam(FShowConnectedClients, const TSharedRef<IConcertServerSession>&);
+
+	FLiveConcertSessionTab(TSharedRef<IConcertServerSession> InspectedSession,
+		TSharedRef<IConcertSyncServer> SyncServer,
+		TAttribute<TSharedRef<SWindow>> ConstructUnderWindow,
+		FShowConnectedClients OnConnectedClientsClicked
+		);
 
 protected:
 
@@ -38,6 +44,9 @@ private:
 	/** Used to determine owning window during construction */
 	TAttribute<TSharedRef<SWindow>> ConstructUnderWindow;
 
+	/** Shows all connected clients for this session */
+	FShowConnectedClients OnConnectedClientsClicked;
+
 	/** Manages the session history widget */
 	const TSharedRef<FLiveServerSessionHistoryController> SessionHistoryController;
 	/** Manages the package viewer widget */
diff --git a/Engine/Plugins/Developer/Concert/ConcertApp/MultiUserServer/Source/MultiUserServer/Private/Widgets/SessionTabs/Live/SConcertLiveSessionTabView.cpp b/Engine/Plugins/Developer/Concert/ConcertApp/MultiUserServer/Source/MultiUserServer/Private/Widgets/SessionTabs/Live/SConcertLiveSessionTabView.cpp
--- a/Engine/Plugins/Developer/Concert/ConcertApp/MultiUserServer/Source/MultiUserServer/Private/Widgets/SessionTabs/Live/SConcertLiveSessionTabView.cpp
+++ b/Engine/Plugins/Developer/Concert/ConcertApp/MultiUserServer/Source/MultiUserServer/Private/Widgets/SessionTabs/Live/SConcertLiveSessionTabView.cpp
@@ -3,14 +3,16 @@
 #include "SConcertLiveSessionTabView.h"
 
 #include "Framework/Docking/TabManager.h"
-#include "Session/History/SSessionHistory.h"
 #include "Widgets/Docking/SDockTab.h"
+#include "Session/History/SSessionHistory.h"
+#include "Widgets/Images/SImage.h"
+#include "Widgets/Input/SButton.h"
+#include "Widgets/SBoxPanel.h"
 
-#define LOCTEXT_NAMESPACE "SConcertSessionInspector"
+#define LOCTEXT_NAMESPACE "MultiUserServerUI.SConcertSessionInspector"
 
 const FName SConcertLiveSessionTabView::HistoryTabId("HistoryTabId");
 const FName SConcertLiveSessionTabView::SessionContentTabId("SessionContentTabId");
-const FName SConcertLiveSessionTabView::ConnectionMonitorTabId("ConnectionMonitorTabId");
 
 void SConcertLiveSessionTabView::Construct(const FArguments& InArgs, const FRequiredWidgets& InRequiredArgs, FName StatusBarId)
 {
@@ -22,7 +24,23 @@ void SConcertLiveSessionTabView::Construct(const FArguments& InArgs, const FRequ
 		{
 			CreateTabs(InTabManager, InLayout, InRequiredArgs);
 		}))
-		.LayoutName("ConcertSessionInspector_v0.2"),
+		.OverlayTabs_Lambda([this, &InArgs](const TSharedRef<SWidget>& Tabs)
+		{
+			return SNew(SVerticalBox)
+				+SVerticalBox::Slot()
+				.AutoHeight()
+				.Padding(3.f)
+				[
+					CreateToolbar(InArgs)
+				]
+			
+				+SVerticalBox::Slot()
+				.FillHeight(1.f)
+				[
+					Tabs
+				];
+		})
+		.LayoutName("ConcertSessionInspector_v0.3"),
 		StatusBarId
 		);
 }
@@ -47,21 +65,9 @@ TSharedRef<SWidget> SConcertLiveSessionTabView::CreateTabs(const TSharedRef<FTab
 			)
 			->Split
 			(
-				FTabManager::NewSplitter()
+				FTabManager::NewStack()
 				->SetSizeCoefficient(0.5f)
-				->SetOrientation(Orient_Vertical)
-				->Split
-				(
-					FTabManager::NewStack()
-					->SetSizeCoefficient(0.5f)
-					->AddTab(SessionContentTabId, ETabState::OpenedTab)
-				)
-				->Split
-				(
-					FTabManager::NewStack()
-					->SetSizeCoefficient(0.5f)
-					->AddTab(ConnectionMonitorTabId, ETabState::OpenedTab)
-				)
+				->AddTab(SessionContentTabId, ETabState::OpenedTab)
 			)
 	);
 
@@ -88,4 +94,43 @@ TSharedRef<SDockTab> SConcertLiveSessionTabView::SpawnSessionContent(const FSpaw
 		];
 }
 
+TSharedRef<SWidget> SConcertLiveSessionTabView::CreateToolbar(const FArguments& InArgs)
+{
+	return SNew(SHorizontalBox)
+
+		+SHorizontalBox::Slot()
+		.AutoWidth()
+		[
+			SNew(SButton)
+			.ButtonStyle(FAppStyle::Get(), "SimpleButton")
+			.ToolTipText(LOCTEXT("ConnectedClientsTooltip", "Opens Clients tab and shows all clients connected to this session."))
+			.ContentPadding(FMargin(1, 0))
+			.OnClicked_Lambda([Callback = InArgs._OnConnectedClientsClicked]()
+			{
+				Callback.ExecuteIfBound();
+				return FReply::Handled();
+			})
+			[
+				SNew(SHorizontalBox)
+				+SHorizontalBox::Slot()
+				.AutoWidth()
+				.VAlign(VAlign_Center)
+				.Padding(4.0, 0.0f)
+				[
+					SNew(SImage)
+					.ColorAndOpacity(FSlateColor::UseForeground())
+					.Image(FAppStyle::Get().GetBrush("Icons.Layout"))
+				]
+				+SHorizontalBox::Slot()
+				.VAlign(VAlign_Center)
+				.Padding(4.0, 0.0f)
+				[
+					SNew(STextBlock)
+					.Text(LOCTEXT("ConnectedClients", "Connected Clients"))
+					.ColorAndOpacity(FSlateColor::UseForeground())
+				]
+			]
+		];
+}
+
 #undef LOCTEXT_NAMESPACE
diff --git a/Engine/Plugins/Developer/Concert/ConcertApp/MultiUserServer/Source/MultiUserServer/Private/Widgets/SessionTabs/Live/SConcertLiveSessionTabView.h b/Engine/Plugins/Developer/Concert/ConcertApp/MultiUserServer/Source/MultiUserServer/Private/Widgets/SessionTabs/Live/SConcertLiveSessionTabView.h
--- a/Engine/Plugins/Developer/Concert/ConcertApp/MultiUserServer/Source/MultiUserServer/Private/Widgets/SessionTabs/Live/SConcertLiveSessionTabView.h
+++ b/Engine/Plugins/Developer/Concert/ConcertApp/MultiUserServer/Source/MultiUserServer/Private/Widgets/SessionTabs/Live/SConcertLiveSessionTabView.h
@@ -44,10 +44,10 @@ public:
 
 	static const FName HistoryTabId;
 	static const FName SessionContentTabId;
-	static const FName ConnectionMonitorTabId;
 
 	SLATE_BEGIN_ARGS(SConcertLiveSessionTabView) {}
 		SLATE_NAMED_SLOT(FArguments, StatusBar)
+		SLATE_EVENT(FSimpleDelegate, OnConnectedClientsClicked)
 	SLATE_END_ARGS()
 
 	void Construct(const FArguments& InArgs, const FRequiredWidgets& InRequiredArgs, FName StatusBarId);
@@ -58,4 +58,6 @@ private:
 	TSharedRef<SWidget> CreateTabs(const TSharedRef<FTabManager>& InTabManager, const TSharedRef<FTabManager::FLayout>& InLayout, const FRequiredWidgets& InRequiredArgs);
 	TSharedRef<SDockTab> SpawnActivityHistory(const FSpawnTabArgs& Args, TSharedRef<SSessionHistory> SessionHistory);
 	TSharedRef<SDockTab> SpawnSessionContent(const FSpawnTabArgs& Args, TSharedRef<SConcertSessionPackageViewer> PackageViewer);
+
+	TSharedRef<SWidget> CreateToolbar(const FArguments& InArgs);
 };
diff --git a/Engine/Plugins/Developer/Concert/ConcertApp/MultiUserServer/Source/MultiUserServer/Private/Window/ConcertServerWindowController.cpp b/Engine/Plugins/Developer/Concert/ConcertApp/MultiUserServer/Source/MultiUserServer/Private/Window/ConcertServerWindowController.cpp
--- a/Engine/Plugins/Developer/Concert/ConcertApp/MultiUserServer/Source/MultiUserServer/Private/Window/ConcertServerWindowController.cpp
+++ b/Engine/Plugins/Developer/Concert/ConcertApp/MultiUserServer/Source/MultiUserServer/Private/Window/ConcertServerWindowController.cpp
@@ -109,7 +109,12 @@ TSharedPtr<FConcertSessionTabBase> FConcertServerWindowController::GetOrRegister
 	
 	if (const TSharedPtr<IConcertServerSession> Session = ServerInstance->GetConcertServer()->GetLiveSession(SessionId))
 	{
-		const TSharedRef<FLiveConcertSessionTab> SessionTab = MakeShared<FLiveConcertSessionTab>(Session.ToSharedRef(), ServerInstance, RootWindow.ToSharedRef());
+		const TSharedRef<FLiveConcertSessionTab> SessionTab = MakeShared<FLiveConcertSessionTab>(
+			Session.ToSharedRef(),
+			ServerInstance,
+			RootWindow.ToSharedRef(),
+			FLiveConcertSessionTab::FShowConnectedClients::CreateSP(this, &FConcertServerWindowController::ShowConnectedClients)
+			);
 		RegisteredSessions.Add(SessionId, SessionTab);
 		return SessionTab;
 	}
@@ -156,6 +161,11 @@ void FConcertServerWindowController::OnArchivedSessionDestroyed(const IConcertSe
 	DestroySessionTab(InArchivedSessionId);
 }
 
+void FConcertServerWindowController::ShowConnectedClients(const TSharedRef<IConcertServerSession>& ServerSession)
+{
+	ClientsController->ShowConnectedClients(ServerSession->GetId());
+}
+
 void FConcertServerWindowController::OnWindowClosed(const TSharedRef<SWindow>& Window)
 {
 	SaveLayout();
diff --git a/Engine/Plugins/Developer/Concert/ConcertApp/MultiUserServer/Source/MultiUserServer/Private/Window/ConcertServerWindowController.h b/Engine/Plugins/Developer/Concert/ConcertApp/MultiUserServer/Source/MultiUserServer/Private/Window/ConcertServerWindowController.h
--- a/Engine/Plugins/Developer/Concert/ConcertApp/MultiUserServer/Source/MultiUserServer/Private/Window/ConcertServerWindowController.h
+++ b/Engine/Plugins/Developer/Concert/ConcertApp/MultiUserServer/Source/MultiUserServer/Private/Window/ConcertServerWindowController.h
@@ -79,6 +79,8 @@ private:
 	void UnregisterFromSessionDestructionEvents() const;
 	void OnLiveSessionDestroyed(const IConcertServer&, TSharedRef<IConcertServerSession> InLiveSession);
 	void OnArchivedSessionDestroyed(const IConcertServer&, const FGuid& InArchivedSessionId);
+
+	void ShowConnectedClients(const TSharedRef<IConcertServerSession>& ServerSession);
 	
 	void OnWindowClosed(const TSharedRef<SWindow>& Window);
 	void SaveLayout() const;
