Filter logs by size

diff --git a/Engine/Plugins/Developer/Concert/ConcertApp/MultiUserServer/Source/MultiUserServer/Private/Widgets/Clients/Logging/Filter/ConcertFrontendLogFilter_Size.cpp b/Engine/Plugins/Developer/Concert/ConcertApp/MultiUserServer/Source/MultiUserServer/Private/Widgets/Clients/Logging/Filter/ConcertFrontendLogFilter_Size.cpp
new file mode 100644
--- /dev/null
+++ b/Engine/Plugins/Developer/Concert/ConcertApp/MultiUserServer/Source/MultiUserServer/Private/Widgets/Clients/Logging/Filter/ConcertFrontendLogFilter_Size.cpp
@@ -0,0 +1,162 @@
+// Copyright Epic Games, Inc. All Rights Reserved.
+
+#include "ConcertFrontendLogFilter_Size.h"
+
+#include "ConcertFrontendUtils.h"
+#include "ConcertTransportEvents.h"
+#include "SSimpleComboButton.h"
+
+#include "Framework/MultiBox/MultiBoxBuilder.h"
+#include "Widgets/SBoxPanel.h"
+#include "Widgets/Input/SNumericEntryBox.h"
+#include "Widgets/Text/STextBlock.h"
+
+#define LOCTEXT_NAMESPACE "UnrealMultiUserUI.FConcertLogFilter_Size"
+
+bool FConcertLogFilter_Size::PassesFilter(const FConcertLog& InItem) const
+{
+	check(InItem.CustomPayloadUncompressedByteSize >= 0);
+	// Note: This only filters activities events - they all use custom events.
+	// The filter's default value is to show everything 0 <= x so it will show sync events as well.
+	const uint32 ComparisionSizeInBytes = FUnitConversion::Convert(SizeInBytes, DataUnit, EUnit::Bytes);
+	switch (FilterMode)
+	{
+	case ESizeFilterMode::LessThanOrEqual:
+		return ComparisionSizeInBytes >= static_cast<uint32>(InItem.CustomPayloadUncompressedByteSize);
+	case ESizeFilterMode::BiggerThanOrEqual:
+		return ComparisionSizeInBytes <= static_cast<uint32>(InItem.CustomPayloadUncompressedByteSize);
+	default:
+		checkNoEntry();
+		return false;
+	}
+}
+
+void FConcertLogFilter_Size::AdvanceFilterMode()
+{
+	switch (FilterMode)
+	{
+	case ESizeFilterMode::LessThanOrEqual:
+		FilterMode = ESizeFilterMode::BiggerThanOrEqual;
+		break;
+	case ESizeFilterMode::BiggerThanOrEqual:
+		FilterMode = ESizeFilterMode::LessThanOrEqual;
+		break;
+	default:
+		checkNoEntry();
+	}
+
+	OnChanged().Broadcast();
+}
+
+void FConcertLogFilter_Size::SetSizeInBytes(uint32 NewSizeInBytes)
+{
+	if (NewSizeInBytes != SizeInBytes)
+	{
+		SizeInBytes = NewSizeInBytes;
+		OnChanged().Broadcast();
+	}
+}
+
+void FConcertLogFilter_Size::SetDataUnit(EUnit NewUnit)
+{
+	if (DataUnit != NewUnit && ensure(GetAllowedUnits().Contains(NewUnit)))
+	{
+		DataUnit = NewUnit;
+		OnChanged().Broadcast();
+	}
+}
+
+FConcertFrontendLogFilter_Size::FConcertFrontendLogFilter_Size()
+{
+	ChildSlot = SNew(SHorizontalBox)
+		+SHorizontalBox::Slot()
+		.AutoWidth()
+		.HAlign(HAlign_Center)
+		[
+			SNew(SButton)
+			.OnClicked_Lambda([this]()
+			{
+				Implementation.AdvanceFilterMode();
+				return FReply::Handled();
+			})
+			.ToolTipText_Lambda([this]()
+			{
+				switch (Implementation.GetFilterMode())
+				{
+				case ESizeFilterMode::BiggerThanOrEqual: return FText::Format(LOCTEXT("BiggerThanOrEqual.ToolTipFmt", "Size >= {0}"), GetSizeAndUnitAsText());
+					case ESizeFilterMode::LessThanOrEqual: return FText::Format(LOCTEXT("LessThanOrEqual.ToolTipFmt", "Size <= {0}"), GetSizeAndUnitAsText());
+					default: return FText::GetEmpty();
+				}
+			})
+			[
+				SNew(STextBlock)
+				.Text_Lambda([this]()
+				{
+					switch (Implementation.GetFilterMode())
+					{
+						case ESizeFilterMode::BiggerThanOrEqual: return LOCTEXT("BiggerThanOrEqual.Text", ">=") ;
+						case ESizeFilterMode::LessThanOrEqual: return LOCTEXT("LessThanOrEqual.Text", "<=") ;
+						default: return FText::GetEmpty();
+					}
+				})
+			]
+		]
+	
+		+SHorizontalBox::Slot()
+		.AutoWidth()
+		[
+			SNew(SNumericEntryBox<uint32>)
+			.AllowSpin(true)
+			.MinDesiredValueWidth(30)
+			.MaxSliderValue(1000)
+			.OnValueChanged_Lambda([this](uint32 NewValue){ Implementation.SetSizeInBytes(NewValue); })
+			.Value_Lambda([this](){ return Implementation.GetSizeInBytes(); })
+		]
+
+		+SHorizontalBox::Slot()
+		.AutoWidth()
+		[
+			SNew(SSimpleComboButton)
+			.HasDownArrow(true)
+			.Text_Lambda([this](){ return FText::FromString(FUnitConversion::GetUnitDisplayString(Implementation.GetDataUnit())); })
+			.MenuContent()
+			[
+				MakeDataUnitMenu()
+			]
+		];
+}
+
+TSharedRef<SWidget> FConcertFrontendLogFilter_Size::MakeDataUnitMenu()
+{
+	FMenuBuilder MenuBuilder(true, nullptr);
+
+	for (const EUnit Unit : Implementation.GetAllowedUnits())
+	{
+		MenuBuilder.AddMenuEntry(
+			FText::FromString(FUnitConversion::GetUnitDisplayString(Unit)),
+			FText::GetEmpty(),
+			FSlateIcon(),
+			FUIAction(
+				FExecuteAction::CreateLambda([this, Unit]()
+				{
+					Implementation.SetDataUnit(Unit);
+				}),
+				FCanExecuteAction::CreateLambda([] { return true; })),
+			NAME_None,
+			EUserInterfaceActionType::Button
+			);
+	}
+	
+	return MenuBuilder.MakeWidget();
+}
+
+FText FConcertFrontendLogFilter_Size::GetSizeAndUnitAsText() const
+{
+	return FText::Format(
+		LOCTEXT("SizeAndUnitAsTextFmt", "{0} {1}"),
+		FUnitConversion::Convert(Implementation.GetSizeInBytes(), Implementation.GetDataUnit(), EUnit::Bytes),
+		FText::FromString(FUnitConversion::GetUnitDisplayString(Implementation.GetDataUnit()))
+		);
+}
+
+#undef LOCTEXT_NAMESPACE
diff --git a/Engine/Plugins/Developer/Concert/ConcertApp/MultiUserServer/Source/MultiUserServer/Private/Widgets/Clients/Logging/Filter/ConcertFrontendLogFilter_Size.h b/Engine/Plugins/Developer/Concert/ConcertApp/MultiUserServer/Source/MultiUserServer/Private/Widgets/Clients/Logging/Filter/ConcertFrontendLogFilter_Size.h
new file mode 100644
--- /dev/null
+++ b/Engine/Plugins/Developer/Concert/ConcertApp/MultiUserServer/Source/MultiUserServer/Private/Widgets/Clients/Logging/Filter/ConcertFrontendLogFilter_Size.h
@@ -0,0 +1,53 @@
+// Copyright Epic Games, Inc. All Rights Reserved.
+
+#pragma once
+
+#include "CoreMinimal.h"
+#include "ConcertFrontendLogFilter.h"
+#include "Math/UnitConversion.h"
+
+enum class ESizeFilterMode
+{
+	/** Allow logs bigger than or equal to the specified size */
+	BiggerThanOrEqual,
+	/** Allow logs smaller than or equal to the specified size */
+	LessThanOrEqual
+};
+
+/** Filters based on the log's size */
+class FConcertLogFilter_Size : public FConcertLogFilter
+{
+public:
+	
+	//~ Begin FConcertLogFilter Interface
+	virtual bool PassesFilter(const FConcertLog& InItem) const override;
+	//~ End FConcertLogFilter Interface
+
+	void AdvanceFilterMode();
+	void SetSizeInBytes(uint32 NewSizeInBytes);
+	void SetDataUnit(EUnit NewUnit);
+	
+	ESizeFilterMode GetFilterMode() const { return FilterMode; }
+	uint32 GetSizeInBytes() const { return SizeInBytes; }
+	EUnit GetDataUnit() const { return DataUnit; }
+	TSet<EUnit> GetAllowedUnits() const { return { EUnit::Bytes, EUnit::Kilobytes, EUnit::Megabytes }; }
+	
+private:
+
+	ESizeFilterMode FilterMode = ESizeFilterMode::BiggerThanOrEqual;
+	uint32 SizeInBytes = 0;
+	EUnit DataUnit = EUnit::Bytes;
+};
+
+class FConcertFrontendLogFilter_Size : public TConcertFrontendLogFilterAggregate<FConcertLogFilter_Size>
+{
+	using Super = TConcertFrontendLogFilterAggregate<FConcertLogFilter_Size>;
+public:
+	
+	FConcertFrontendLogFilter_Size();
+
+private:
+
+	TSharedRef<SWidget> MakeDataUnitMenu();
+	FText GetSizeAndUnitAsText() const;
+};
diff --git a/Engine/Plugins/Developer/Concert/ConcertApp/MultiUserServer/Source/MultiUserServer/Private/Widgets/Clients/Logging/Filter/ConcertLogFilter_FrontendRoot.cpp b/Engine/Plugins/Developer/Concert/ConcertApp/MultiUserServer/Source/MultiUserServer/Private/Widgets/Clients/Logging/Filter/ConcertLogFilter_FrontendRoot.cpp
--- a/Engine/Plugins/Developer/Concert/ConcertApp/MultiUserServer/Source/MultiUserServer/Private/Widgets/Clients/Logging/Filter/ConcertLogFilter_FrontendRoot.cpp
+++ b/Engine/Plugins/Developer/Concert/ConcertApp/MultiUserServer/Source/MultiUserServer/Private/Widgets/Clients/Logging/Filter/ConcertLogFilter_FrontendRoot.cpp
@@ -8,6 +8,7 @@
 #include "ConcertFrontendLogFilter_MessageType.h"
 #include "ConcertFrontendLogFilter_TextSearch.h"
 #include "ConcertFrontendLogFilter_Time.h"
+#include "ConcertFrontendLogFilter_Size.h"
 #include "Algo/AllOf.h"
 
 #include "Algo/AnyOf.h"
@@ -93,7 +94,8 @@ namespace UE::MultiUserServer
 				MakeShared<FConcertFrontendLogFilter_MessageAction>(),
 				MakeShared<FConcertFrontendLogFilter_MessageType>(),
 				MakeShared<FConcertFrontendLogFilter_Time>(ETimeFilter::AllowAfter),
-				MakeShared<FConcertFrontendLogFilter_Time>(ETimeFilter::AllowBefore)
+				MakeShared<FConcertFrontendLogFilter_Time>(ETimeFilter::AllowBefore),
+				MakeShared<FConcertFrontendLogFilter_Size>()
 			};
 		}
 	}
diff --git a/Engine/Plugins/Developer/Concert/ConcertApp/MultiUserServer/Source/MultiUserServer/Private/Widgets/Clients/Logging/SConcertTransportLog.cpp b/Engine/Plugins/Developer/Concert/ConcertApp/MultiUserServer/Source/MultiUserServer/Private/Widgets/Clients/Logging/SConcertTransportLog.cpp
--- a/Engine/Plugins/Developer/Concert/ConcertApp/MultiUserServer/Source/MultiUserServer/Private/Widgets/Clients/Logging/SConcertTransportLog.cpp
+++ b/Engine/Plugins/Developer/Concert/ConcertApp/MultiUserServer/Source/MultiUserServer/Private/Widgets/Clients/Logging/SConcertTransportLog.cpp
@@ -118,18 +118,16 @@ TSharedRef<SHeaderRow> SConcertTransportLog::CreateHeaderRow()
 			.ToolTipText(LOCTEXT("AvatarColumnToolTipText", "The colour of the avatar is affected by log"))
 		;
 
-	// The following properties are useless to the user
-	const TSet<const FProperty*> SkippedProperties = {
-		// Used by code only
-		FConcertLog::StaticStruct()->FindPropertyByName(GET_MEMBER_NAME_CHECKED(FConcertLog, CustomPayloadUncompressedByteSize))
+	const TMap<FName, FString> ColumnNameOverrides = {
+		{ GET_MEMBER_NAME_CHECKED(FConcertLog, CustomPayloadUncompressedByteSize), TEXT("Size") }
 	};
-	
 	for (TFieldIterator<const FProperty> PropertyIt(FConcertLog::StaticStruct()); PropertyIt; ++PropertyIt)
 	{
-		if (!PropertyIt->HasAnyPropertyFlags(CPF_Transient) && !SkippedProperties.Contains(*PropertyIt))
+		if (!PropertyIt->HasAnyPropertyFlags(CPF_Transient))
 		{
-			const FString PropertyName = PropertyIt->GetAuthoredName();
 			const FName ColumnId = PropertyIt->GetFName();
+			const FString* LabelOverride = ColumnNameOverrides.Find(ColumnId);
+			const FString PropertyName = LabelOverride ? *LabelOverride : PropertyIt->GetAuthoredName();
 			HeaderRow->AddColumn(SHeaderRow::FColumn::FArguments()
 				.ColumnId(ColumnId)
 				.DefaultLabel(FText::FromString(PropertyName))
@@ -143,6 +141,19 @@ TSharedRef<SHeaderRow> SConcertTransportLog::CreateHeaderRow()
 	}
 
 	TGuardValue<bool> DoNotSave(bIsUpdatingColumnVisibility, true);
+	RestoreDefaultColumnVisiblities();
+	
+	return HeaderRow.ToSharedRef();
+}
+
+TSharedRef<ITableRow> SConcertTransportLog::OnGenerateActivityRowWidget(TSharedPtr<FConcertLogEntry> Item, const TSharedRef<STableViewBase>& OwnerTable) const
+{
+	return SNew(SConcertTransportLogRow, Item, OwnerTable, LogTokenizer.ToSharedRef(), HighlightText.ToSharedRef())
+		.GetClientInfo(GetClientInfoFunc);
+}
+
+void SConcertTransportLog::RestoreDefaultColumnVisiblities()
+{
 	const TSet<FName> HiddenByDefault = {
 		GET_MEMBER_NAME_CHECKED(FConcertLog, Frame),
 		GET_MEMBER_NAME_CHECKED(FConcertLog, MessageId),
@@ -155,14 +166,6 @@ TSharedRef<SHeaderRow> SConcertTransportLog::CreateHeaderRow()
 	{
 		HeaderRow->SetShowGeneratedColumn(ColumnID, false);
 	}
-	
-	return HeaderRow.ToSharedRef();
-}
-
-TSharedRef<ITableRow> SConcertTransportLog::OnGenerateActivityRowWidget(TSharedPtr<FConcertLogEntry> Item, const TSharedRef<STableViewBase>& OwnerTable) const
-{
-	return SNew(SConcertTransportLogRow, Item, OwnerTable, LogTokenizer.ToSharedRef(), HighlightText.ToSharedRef())
-		.GetClientInfo(GetClientInfoFunc);
 }
 
 void SConcertTransportLog::ExtendViewOptions(FMenuBuilder& MenuBuilder)
@@ -197,8 +200,17 @@ void SConcertTransportLog::ExtendViewOptions(FMenuBuilder& MenuBuilder)
 			EUserInterfaceActionType::ToggleButton
 		);
 
-					
 	MenuBuilder.AddSeparator();
+	MenuBuilder.AddMenuEntry(
+			LOCTEXT("RestoreDefaultColumnVisibility", "Restore columns visibility"),
+			FText::GetEmpty(),
+			FSlateIcon(),
+			FUIAction(
+				FExecuteAction::CreateSP(this, &SConcertTransportLog::RestoreDefaultColumnVisiblities),
+				FCanExecuteAction::CreateLambda([] { return true; })),
+			NAME_None,
+			EUserInterfaceActionType::Button
+		);
 	UE::ConcertSharedSlate::AddEntriesForShowingHiddenRows(HeaderRow.ToSharedRef(), MenuBuilder);
 }
 
diff --git a/Engine/Plugins/Developer/Concert/ConcertApp/MultiUserServer/Source/MultiUserServer/Private/Widgets/Clients/Logging/SConcertTransportLog.h b/Engine/Plugins/Developer/Concert/ConcertApp/MultiUserServer/Source/MultiUserServer/Private/Widgets/Clients/Logging/SConcertTransportLog.h
--- a/Engine/Plugins/Developer/Concert/ConcertApp/MultiUserServer/Source/MultiUserServer/Private/Widgets/Clients/Logging/SConcertTransportLog.h
+++ b/Engine/Plugins/Developer/Concert/ConcertApp/MultiUserServer/Source/MultiUserServer/Private/Widgets/Clients/Logging/SConcertTransportLog.h
@@ -75,6 +75,7 @@ private:
 	TSharedRef<SHeaderRow> CreateHeaderRow();
 	TSharedRef<ITableRow> OnGenerateActivityRowWidget(TSharedPtr<FConcertLogEntry> Item, const TSharedRef<STableViewBase>& OwnerTable) const;
 
+	void RestoreDefaultColumnVisiblities();
 	void ExtendViewOptions(FMenuBuilder& MenuBuilder);
 	void OnFilterMenuChecked();
 	void OnPageViewChanged(const TArray<TSharedPtr<FConcertLogEntry>>&);
diff --git a/Engine/Plugins/Developer/Concert/ConcertApp/MultiUserServer/Source/MultiUserServer/Private/Widgets/Clients/Logging/Util/ConcertLogTokenizer.cpp b/Engine/Plugins/Developer/Concert/ConcertApp/MultiUserServer/Source/MultiUserServer/Private/Widgets/Clients/Logging/Util/ConcertLogTokenizer.cpp
--- a/Engine/Plugins/Developer/Concert/ConcertApp/MultiUserServer/Source/MultiUserServer/Private/Widgets/Clients/Logging/Util/ConcertLogTokenizer.cpp
+++ b/Engine/Plugins/Developer/Concert/ConcertApp/MultiUserServer/Source/MultiUserServer/Private/Widgets/Clients/Logging/Util/ConcertLogTokenizer.cpp
@@ -5,13 +5,15 @@
 #include "ConcertFrontendUtils.h"
 #include "ConcertTransportEvents.h"
 #include "MessageTypeUtils.h"
+#include "Math/UnitConversion.h"
 #include "Settings/ConcertTransportLogSettings.h"
 
 FConcertLogTokenizer::FConcertLogTokenizer()
 {
 	TokenizerFunctions = {
 			{ FConcertLog::StaticStruct()->FindPropertyByName(GET_MEMBER_NAME_CHECKED(FConcertLog, Timestamp)), [this](const FConcertLog& Log) { return TokenizeTimestamp(Log); } },
-			{ FConcertLog::StaticStruct()->FindPropertyByName(GET_MEMBER_NAME_CHECKED(FConcertLog, MessageTypeName)), [this](const FConcertLog& Log) { return TokenizeMessageTypeName(Log); } }
+			{ FConcertLog::StaticStruct()->FindPropertyByName(GET_MEMBER_NAME_CHECKED(FConcertLog, MessageTypeName)), [this](const FConcertLog& Log) { return TokenizeMessageTypeName(Log); } },
+			{ FConcertLog::StaticStruct()->FindPropertyByName(GET_MEMBER_NAME_CHECKED(FConcertLog, CustomPayloadUncompressedByteSize)), [this](const FConcertLog& Log) { return TokenizeCustomPayloadUncompressedByteSize(Log); } }
 		};
 }
 
@@ -34,6 +36,13 @@ FString FConcertLogTokenizer::TokenizeMessageTypeName(const FConcertLog& Data) c
 	return UE::MultiUserServer::MessageTypeUtils::SanitizeMessageTypeName(Data.MessageTypeName);
 }
 
+FString FConcertLogTokenizer::TokenizeCustomPayloadUncompressedByteSize(const FConcertLog& Data) const
+{
+	// So changes to the type are automatically propagated here
+	const FNumericUnit<int32> DisplayUnit = FUnitConversion::QuantizeUnitsToBestFit(Data.CustomPayloadUncompressedByteSize, EUnit::Bytes);
+	return FString::Printf(TEXT("%d %s"), DisplayUnit.Value, FUnitConversion::GetUnitDisplayString(DisplayUnit.Units));
+}
+
 FString FConcertLogTokenizer::TokenizeUsingPropertyExport(const FConcertLog& Data, const FProperty& ConcertLogProperty) const
 {
 	FString Exported;
diff --git a/Engine/Plugins/Developer/Concert/ConcertApp/MultiUserServer/Source/MultiUserServer/Private/Widgets/Clients/Logging/Util/ConcertLogTokenizer.h b/Engine/Plugins/Developer/Concert/ConcertApp/MultiUserServer/Source/MultiUserServer/Private/Widgets/Clients/Logging/Util/ConcertLogTokenizer.h
--- a/Engine/Plugins/Developer/Concert/ConcertApp/MultiUserServer/Source/MultiUserServer/Private/Widgets/Clients/Logging/Util/ConcertLogTokenizer.h
+++ b/Engine/Plugins/Developer/Concert/ConcertApp/MultiUserServer/Source/MultiUserServer/Private/Widgets/Clients/Logging/Util/ConcertLogTokenizer.h
@@ -18,6 +18,7 @@ public:
 
 	FString TokenizeTimestamp(const FConcertLog& Data) const;
 	FString TokenizeMessageTypeName(const FConcertLog& Data) const;
+	FString TokenizeCustomPayloadUncompressedByteSize(const FConcertLog& Data) const;
 	FString TokenizeUsingPropertyExport(const FConcertLog& Data, const FProperty& ConcertLogProperty) const;
 
 private:
