Closing and re-opening client log loses logs

diff --git a/Engine/Plugins/Developer/Concert/ConcertApp/MultiUserServer/Source/MultiUserServer/Private/Widgets/Clients/Logging/Filter/FilteredConcertLogList.cpp b/Engine/Plugins/Developer/Concert/ConcertApp/MultiUserServer/Source/MultiUserServer/Private/Widgets/Clients/Logging/Filter/FilteredConcertLogList.cpp
--- a/Engine/Plugins/Developer/Concert/ConcertApp/MultiUserServer/Source/MultiUserServer/Private/Widgets/Clients/Logging/Filter/FilteredConcertLogList.cpp
+++ b/Engine/Plugins/Developer/Concert/ConcertApp/MultiUserServer/Source/MultiUserServer/Private/Widgets/Clients/Logging/Filter/FilteredConcertLogList.cpp
@@ -15,6 +15,8 @@ FFilteredConcertLogList::FFilteredConcertLogList(TSharedRef<IConcertLogSource> I
 	{
 		Filter->OnChanged().AddRaw(this, &FFilteredConcertLogList::RebuildFilteredResult);
 	}
+
+	RebuildFilteredResult();
 }
 
 FFilteredConcertLogList::~FFilteredConcertLogList()
@@ -72,6 +74,7 @@ FPagedFilteredConcertLogList::FPagedFilteredConcertLogList(TSharedRef<IConcertLo
 	checkf(InitialLogsPerPage >= 1, TEXT("Unreasonable page size"));
 
 	OnLogListChanged().AddRaw(this, &FPagedFilteredConcertLogList::HandeOnLogListChanged);
+	CheckAndConditionallyPopulatePage();
 }
 
 void FPagedFilteredConcertLogList::SetLogsPerPage(uint16 NewLogsPerPage)
