Fix non-unity issue

diff --git a/Engine/Plugins/Developer/Concert/ConcertApp/MultiUserServer/Source/MultiUserServer/Private/Widgets/Clients/Logging/Filter/ConcertFrontendLogFilter_Size.cpp b/Engine/Plugins/Developer/Concert/ConcertApp/MultiUserServer/Source/MultiUserServer/Private/Widgets/Clients/Logging/Filter/ConcertFrontendLogFilter_Size.cpp
--- a/Engine/Plugins/Developer/Concert/ConcertApp/MultiUserServer/Source/MultiUserServer/Private/Widgets/Clients/Logging/Filter/ConcertFrontendLogFilter_Size.cpp
+++ b/Engine/Plugins/Developer/Concert/ConcertApp/MultiUserServer/Source/MultiUserServer/Private/Widgets/Clients/Logging/Filter/ConcertFrontendLogFilter_Size.cpp
@@ -8,6 +8,7 @@
 
 #include "Framework/MultiBox/MultiBoxBuilder.h"
 #include "Widgets/SBoxPanel.h"
+#include "Widgets/Input/SButton.h"
 #include "Widgets/Input/SNumericEntryBox.h"
 #include "Widgets/Text/STextBlock.h"
 
