Fix crash when log search yields less pages than currently selected page

diff --git a/Engine/Plugins/Developer/Concert/ConcertApp/MultiUserServer/Source/MultiUserServer/Private/Widgets/Clients/Logging/Filter/FilteredConcertLogList.cpp b/Engine/Plugins/Developer/Concert/ConcertApp/MultiUserServer/Source/MultiUserServer/Private/Widgets/Clients/Logging/Filter/FilteredConcertLogList.cpp
--- a/Engine/Plugins/Developer/Concert/ConcertApp/MultiUserServer/Source/MultiUserServer/Private/Widgets/Clients/Logging/Filter/FilteredConcertLogList.cpp
+++ b/Engine/Plugins/Developer/Concert/ConcertApp/MultiUserServer/Source/MultiUserServer/Private/Widgets/Clients/Logging/Filter/FilteredConcertLogList.cpp
@@ -103,7 +103,16 @@ void FPagedFilteredConcertLogList::SetPage(const uint32 PageIndex)
 
 void FPagedFilteredConcertLogList::HandeOnLogListChanged(const TArray<TSharedPtr<FConcertLogEntry>>& NewFilteredLogList)
 {
-	CheckAndConditionallyPopulatePage();
+	// After filtering there may be less logs > less pages > the current page might have become out of bounds
+	if (CurrentPageIndex >= GetNumPages())
+	{
+		CurrentPageIndex = 0;
+		RepopulatePage();
+	}
+	else
+	{
+		CheckAndConditionallyPopulatePage();
+	}
 }
 
 void FPagedFilteredConcertLogList::RepopulatePage()
@@ -161,7 +170,7 @@ void FPagedFilteredConcertLogList::ForEachLogIndexOnPage(TFunctionRef<void(size_
 {
 	size_t StartIndex, LastIndex;
 	Tie(StartIndex, LastIndex) = GetLogIndicesForPage();
-	const FLogsPerPageCount NumItems = LastIndex - StartIndex + 1;
+	const FLogsPerPageCount NumItems = FMath::Min(LastIndex, LastIndex - StartIndex + 1);
 	for (size_t i = 0; i < NumItems && i < MaxItems; ++i)
 	{
 		Callback(StartIndex + i);
@@ -174,7 +183,7 @@ TTuple<size_t, size_t> FPagedFilteredConcertLogList::GetLogIndicesForPage() cons
 	const size_t LastIndex = FMath::Min<size_t>(
 		// Used when page is full
 		StartIndex + LogsPerPage - 1,
-		// Used when page is not full
+		// Used when last page is not full
 		GetFilteredLogs().Num() == 0 ? 0 : GetFilteredLogs().Num() - 1
 		);
 	check(GetFilteredLogs().Num() == 0 || GetFilteredLogs().IsValidIndex(StartIndex));
