Fix wrong unit of bytes shown in UI

diff --git a/Engine/Plugins/Developer/Concert/ConcertApp/MultiUserServer/Source/MultiUserServer/Private/Widgets/Clients/Browser/SClientNetworkStats.cpp b/Engine/Plugins/Developer/Concert/ConcertApp/MultiUserServer/Source/MultiUserServer/Private/Widgets/Clients/Browser/SClientNetworkStats.cpp
--- a/Engine/Plugins/Developer/Concert/ConcertApp/MultiUserServer/Source/MultiUserServer/Private/Widgets/Clients/Browser/SClientNetworkStats.cpp
+++ b/Engine/Plugins/Developer/Concert/ConcertApp/MultiUserServer/Source/MultiUserServer/Private/Widgets/Clients/Browser/SClientNetworkStats.cpp
@@ -69,7 +69,7 @@ namespace UE::MultiUserServer
 		// Receive
 		{
 			const FNumericUnit<uint64> Unit = FUnitConversion::QuantizeUnitsToBestFit(Statistics.TotalBytesReceived, EUnit::Bytes);
-			const FString DisplayString = FString::Printf(TEXT("%lld %s"), Statistics.TotalBytesReceived, FUnitConversion::GetUnitDisplayString(Unit.Units));
+			const FString DisplayString = FString::Printf(TEXT("%lld %s"), Unit.Value, FUnitConversion::GetUnitDisplayString(Unit.Units));
 			ReceiveText->SetText(FText::FromString(DisplayString));
 		}
 		// RTT
@@ -81,13 +81,13 @@ namespace UE::MultiUserServer
 		// Inflight
 		{
 			const FNumericUnit<uint64> Unit = FUnitConversion::QuantizeUnitsToBestFit(Statistics.BytesInflight, EUnit::Bytes);
-			const FString DisplayString = FString::Printf(TEXT("%lld %s"), Statistics.BytesInflight, FUnitConversion::GetUnitDisplayString(Unit.Units));
+			const FString DisplayString = FString::Printf(TEXT("%lld %s"), Unit.Value, FUnitConversion::GetUnitDisplayString(Unit.Units));
 			InflightText->SetText(FText::FromString(DisplayString));
 		}
 		// Loss
 		{
 			const FNumericUnit<uint64> Unit = FUnitConversion::QuantizeUnitsToBestFit(Statistics.TotalBytesLost, EUnit::Bytes);
-			const FString DisplayString = FString::Printf(TEXT("%lld %s"), Statistics.TotalBytesLost, FUnitConversion::GetUnitDisplayString(Unit.Units));
+			const FString DisplayString = FString::Printf(TEXT("%lld %s"), Unit.Value, FUnitConversion::GetUnitDisplayString(Unit.Units));
 			LossText->SetText(FText::FromString(DisplayString));
 		}
 	}
