Display client / server name instead of endpoint ID

diff --git a/Engine/Plugins/Developer/Concert/ConcertApp/MultiUserServer/Source/MultiUserServer/Private/Widgets/Clients/Logging/SConcertTransportLog.cpp b/Engine/Plugins/Developer/Concert/ConcertApp/MultiUserServer/Source/MultiUserServer/Private/Widgets/Clients/Logging/SConcertTransportLog.cpp
--- a/Engine/Plugins/Developer/Concert/ConcertApp/MultiUserServer/Source/MultiUserServer/Private/Widgets/Clients/Logging/SConcertTransportLog.cpp
+++ b/Engine/Plugins/Developer/Concert/ConcertApp/MultiUserServer/Source/MultiUserServer/Private/Widgets/Clients/Logging/SConcertTransportLog.cpp
@@ -14,6 +14,7 @@
 
 #include "Framework/MultiBox/MultiBoxBuilder.h"
 #include "Styling/AppStyle.h"
+#include "Util/EndpointToUserNameCache.h"
 #include "Widgets/Layout/SBorder.h"
 #include "Widgets/SBoxPanel.h"
 #include "Widgets/SNullWidget.h"
@@ -29,13 +30,12 @@ SConcertTransportLog::~SConcertTransportLog()
 	ConcertTransportEvents::OnConcertTransportLoggingEnabledChangedEvent().RemoveAll(this);
 }
 
-void SConcertTransportLog::Construct(const FArguments& InArgs, TSharedRef<IConcertLogSource> LogSource)
+void SConcertTransportLog::Construct(const FArguments& InArgs, TSharedRef<IConcertLogSource> LogSource, TSharedRef<FEndpointToUserNameCache> InEndpointCache, TSharedRef<FConcertLogTokenizer> InLogTokenizer)
 {
 	PagedLogList = MakeShared<FPagedFilteredConcertLogList>(MoveTemp(LogSource), InArgs._Filter);
-	LogTokenizer = MakeShared<FConcertLogTokenizer>();
+	EndpointCache = MoveTemp(InEndpointCache);
+	LogTokenizer = MoveTemp(InLogTokenizer);
 	HighlightText = MakeShared<FText>();
-
-	GetClientInfoFunc = InArgs._GetClientInfo;
 	
 	ChildSlot
 	[
@@ -119,7 +119,9 @@ TSharedRef<SHeaderRow> SConcertTransportLog::CreateHeaderRow()
 		;
 
 	const TMap<FName, FString> ColumnNameOverrides = {
-		{ GET_MEMBER_NAME_CHECKED(FConcertLog, CustomPayloadUncompressedByteSize), TEXT("Size") }
+		{ GET_MEMBER_NAME_CHECKED(FConcertLog, CustomPayloadUncompressedByteSize), TEXT("Size") },
+		{ GET_MEMBER_NAME_CHECKED(FConcertLog, OriginEndpointId), TEXT("Origin") },
+		{ GET_MEMBER_NAME_CHECKED(FConcertLog, DestinationEndpointId), TEXT("Destination") }
 	};
 	for (TFieldIterator<const FProperty> PropertyIt(FConcertLog::StaticStruct()); PropertyIt; ++PropertyIt)
 	{
@@ -148,8 +150,14 @@ TSharedRef<SHeaderRow> SConcertTransportLog::CreateHeaderRow()
 
 TSharedRef<ITableRow> SConcertTransportLog::OnGenerateActivityRowWidget(TSharedPtr<FConcertLogEntry> Item, const TSharedRef<STableViewBase>& OwnerTable) const
 {
+	const TOptional<FConcertClientInfo> OriginInfo = EndpointCache->GetClientInfo(Item->Log.OriginEndpointId);
+	const TOptional<FConcertClientInfo> DestinationInfo = EndpointCache->GetClientInfo(Item->Log.DestinationEndpointId);
+	const FLinearColor AvatarColor = OriginInfo.IsSet()
+		? OriginInfo->AvatarColor
+		: DestinationInfo.IsSet() ? DestinationInfo->AvatarColor : FLinearColor::Black;
+	
 	return SNew(SConcertTransportLogRow, Item, OwnerTable, LogTokenizer.ToSharedRef(), HighlightText.ToSharedRef())
-		.GetClientInfo(GetClientInfoFunc);
+		.AvatarColor(AvatarColor);
 }
 
 void SConcertTransportLog::RestoreDefaultColumnVisiblities()
@@ -237,7 +245,7 @@ void SConcertTransportLog::OnPageViewChanged(const TArray<TSharedPtr<FConcertLog
 {
 	LogView->RequestListRefresh();
 
-	if (bAutoScroll)
+	if (bAutoScroll && PagedLogList->GetCurrentPage() == PagedLogList->GetNumPages() - 1)
 	{
 		LogView->ScrollToBottom();
 	}
diff --git a/Engine/Plugins/Developer/Concert/ConcertApp/MultiUserServer/Source/MultiUserServer/Private/Widgets/Clients/Logging/SConcertTransportLog.h b/Engine/Plugins/Developer/Concert/ConcertApp/MultiUserServer/Source/MultiUserServer/Private/Widgets/Clients/Logging/SConcertTransportLog.h
--- a/Engine/Plugins/Developer/Concert/ConcertApp/MultiUserServer/Source/MultiUserServer/Private/Widgets/Clients/Logging/SConcertTransportLog.h
+++ b/Engine/Plugins/Developer/Concert/ConcertApp/MultiUserServer/Source/MultiUserServer/Private/Widgets/Clients/Logging/SConcertTransportLog.h
@@ -4,19 +4,19 @@
 
 #include "CoreMinimal.h"
 #include "Filter/FilteredConcertLogList.h"
-#include "TransportLogDelegates.h"
 #include "Widgets/DeclarativeSyntaxSupport.h"
 #include "Widgets/SCompoundWidget.h"
 
-class SOverlay;
-class SPromptConcertLoggingEnabled;
 class FConcertLogTokenizer;
 class FConcertLogFilter_FrontendRoot;
+class FEndpointToUserNameCache;
 class FMenuBuilder;
 class FPagedFilteredConcertLogList;
 class IConcertLogSource;
 class ITableRow;
 class SHeaderRow;
+class SOverlay;
+class SPromptConcertLoggingEnabled;
 class STableViewBase;
 template <typename ItemType> class SListView;
 
@@ -36,12 +36,10 @@ public:
 	{}
 		/** Optional filters to display in UI */
 		SLATE_ARGUMENT(TSharedPtr<FConcertLogFilter_FrontendRoot>, Filter)
-		/** Used to better display client info, e.g. client name and colour */
-		SLATE_EVENT(FGetClientInfo, GetClientInfo)
 	SLATE_END_ARGS()
 	virtual ~SConcertTransportLog() override;
 
-	void Construct(const FArguments& InArgs, TSharedRef<IConcertLogSource> LogSource);
+	void Construct(const FArguments& InArgs, TSharedRef<IConcertLogSource> LogSource, TSharedRef<FEndpointToUserNameCache> InEndpointCache, TSharedRef<FConcertLogTokenizer> InLogTokenizer);
 
 private:
 	
@@ -49,6 +47,8 @@ private:
 	TSharedPtr<SOverlay> EnableLoggingPromptOverlay;
 	/** Reminds the user to enable logging */
 	TSharedPtr<SPromptConcertLoggingEnabled> EnableLoggingPrompt;
+
+	TSharedPtr<FEndpointToUserNameCache> EndpointCache;
 	
 	/** Sorts the log into pages whilst applying filters */
 	TSharedPtr<FPagedFilteredConcertLogList> PagedLogList;
@@ -63,8 +63,6 @@ private:
 	/** Header row of LogView */
 	TSharedPtr<SHeaderRow> HeaderRow;
 
-	FGetClientInfo GetClientInfoFunc;
-
 	/** Whether to automatically scroll to new logs as they come in */
 	bool bAutoScroll = true;
 	/** Whether we currently loading the column visibility - prevents infinite event recursion */
diff --git a/Engine/Plugins/Developer/Concert/ConcertApp/MultiUserServer/Source/MultiUserServer/Private/Widgets/Clients/Logging/SConcertTransportLogRow.cpp b/Engine/Plugins/Developer/Concert/ConcertApp/MultiUserServer/Source/MultiUserServer/Private/Widgets/Clients/Logging/SConcertTransportLogRow.cpp
--- a/Engine/Plugins/Developer/Concert/ConcertApp/MultiUserServer/Source/MultiUserServer/Private/Widgets/Clients/Logging/SConcertTransportLogRow.cpp
+++ b/Engine/Plugins/Developer/Concert/ConcertApp/MultiUserServer/Source/MultiUserServer/Private/Widgets/Clients/Logging/SConcertTransportLogRow.cpp
@@ -14,8 +14,8 @@ void SConcertTransportLogRow::Construct(const FArguments& InArgs, TSharedPtr<FCo
 	LogEntry = MoveTemp(InLogEntry);
 	Tokenizer = MoveTemp(InTokenizer);
 	HighlightText = MoveTemp(InHighlightText);
-
-	GetClientInfoFunc = InArgs._GetClientInfo;
+	
+	AvatarColor = InArgs._AvatarColor;
 	
 	SMultiColumnTableRow<TSharedPtr<FConcertLogEntry>>::Construct(FSuperRowType::FArguments(), InOwnerTableView);
 }
@@ -29,22 +29,13 @@ TSharedRef<SWidget> SConcertTransportLogRow::GenerateWidgetForColumn(const FName
 			.Padding(2, 1)
 			[
 				SNew(SColorBlock)
-				.Color_Lambda([this]()
-				{
-					const TOptional<FConcertClientInfo> ClientInfo = GetClientInfoFromLog();
-					return ClientInfo.IsSet()
-						? ClientInfo->AvatarColor
-						: FLinearColor::Black;
-				})
+				.Color(AvatarColor)
 				.Size(FVector2D(4.f, 16.f))
 			];
 	}
 	
 	using FColumnWidgetFactoryFunc = TSharedRef<SWidget>(SConcertTransportLogRow::*)(const FName& InColumnName);
-	const TMap<FName, FColumnWidgetFactoryFunc> OverrideFactories = {
-		{ GET_MEMBER_NAME_CHECKED(FConcertLog, OriginEndpointId), &SConcertTransportLogRow::CreateClientNameColumn },
-		{ GET_MEMBER_NAME_CHECKED(FConcertLog, DestinationEndpointId), &SConcertTransportLogRow::CreateClientNameColumn }
-	};
+	const TMap<FName, FColumnWidgetFactoryFunc> OverrideFactories = {};
 	const FColumnWidgetFactoryFunc FallbackFactoryFunc = &SConcertTransportLogRow::CreateDefaultColumn;
 
 	const FColumnWidgetFactoryFunc* FactoryFunc = OverrideFactories.Find(ColumnName);
@@ -70,37 +61,3 @@ TSharedRef<SWidget> SConcertTransportLogRow::CreateDefaultColumn(const FName& Pr
 		})
 		.HighlightText_Lambda([this](){ return *HighlightText.Get(); });
 }
-
-TSharedRef<SWidget> SConcertTransportLogRow::CreateClientNameColumn(const FName& PropertyName)
-{
-	const bool bUseOrigin = GET_MEMBER_NAME_CHECKED(FConcertLog, OriginEndpointId) == PropertyName;
-	const bool bUseDestination = GET_MEMBER_NAME_CHECKED(FConcertLog, DestinationEndpointId) == PropertyName;
-	check(bUseOrigin || bUseDestination);
-
-	const FGuid& EndpointId = bUseOrigin
-		? LogEntry->Log.OriginEndpointId
-		: LogEntry->Log.DestinationEndpointId;
-	const TOptional<FConcertClientInfo> ClientInfo = GetClientInfoFunc.IsBound()
-		? GetClientInfoFunc.Execute(EndpointId)
-		: TOptional<FConcertClientInfo>();
-	return SNew(STextBlock)
-		// If there is no client info available when the log is
-		.Text(FText::FromString(ClientInfo.IsSet() ? ClientInfo->DisplayName : EndpointId.ToString()))
-		.HighlightText_Lambda([this](){ return *HighlightText.Get(); });
-}
-
-TOptional<FConcertClientInfo> SConcertTransportLogRow::GetClientInfoFromLog() const
-{
-	if (!GetClientInfoFunc.IsBound())
-	{
-		return {};
-	}
-		
-	const TOptional<FConcertClientInfo> OriginClientInfo = GetClientInfoFunc.IsBound() ? GetClientInfoFunc.Execute(LogEntry->Log.OriginEndpointId) : TOptional<FConcertClientInfo>();
-	const TOptional<FConcertClientInfo> DestinationClientInfo = GetClientInfoFunc.IsBound() ? GetClientInfoFunc.Execute(LogEntry->Log.DestinationEndpointId) : TOptional<FConcertClientInfo>();
-
-	// One of the two is definitely the server. If one of them is a client whose info could be found, this will return it.
-	return OriginClientInfo.IsSet()
-		? OriginClientInfo
-		: DestinationClientInfo;
-}
diff --git a/Engine/Plugins/Developer/Concert/ConcertApp/MultiUserServer/Source/MultiUserServer/Private/Widgets/Clients/Logging/SConcertTransportLogRow.h b/Engine/Plugins/Developer/Concert/ConcertApp/MultiUserServer/Source/MultiUserServer/Private/Widgets/Clients/Logging/SConcertTransportLogRow.h
--- a/Engine/Plugins/Developer/Concert/ConcertApp/MultiUserServer/Source/MultiUserServer/Private/Widgets/Clients/Logging/SConcertTransportLogRow.h
+++ b/Engine/Plugins/Developer/Concert/ConcertApp/MultiUserServer/Source/MultiUserServer/Private/Widgets/Clients/Logging/SConcertTransportLogRow.h
@@ -4,7 +4,6 @@
 
 #include "CoreMinimal.h"
 #include "ConcertLogEntry.h"
-#include "TransportLogDelegates.h"
 #include "Widgets/DeclarativeSyntaxSupport.h"
 #include "Widgets/Views/STableRow.h"
 
@@ -16,8 +15,7 @@ public:
 
 	SLATE_BEGIN_ARGS(SConcertTransportLogRow)
 	{}
-		/** Used to get the avatar's colour */
-		SLATE_EVENT(FGetClientInfo, GetClientInfo)
+		SLATE_ARGUMENT(FLinearColor, AvatarColor)
 	SLATE_END_ARGS()
 
 	void Construct(const FArguments& InArgs, TSharedPtr<FConcertLogEntry> InLogEntry, const TSharedRef<STableViewBase>& InOwnerTableView, TSharedRef<FConcertLogTokenizer> InTokenizer, TSharedRef<FText> InHighlightText);
@@ -34,11 +32,8 @@ private:
 
 	/** Owned by SConcertTransportLog. Updated with the search text. */
 	TSharedPtr<FText> HighlightText;
-	
-	FGetClientInfo GetClientInfoFunc;
+
+	FLinearColor AvatarColor;
 
 	TSharedRef<SWidget> CreateDefaultColumn(const FName& PropertyName);
-	TSharedRef<SWidget> CreateClientNameColumn(const FName& PropertyName);
-	
-	TOptional<FConcertClientInfo> GetClientInfoFromLog() const;
 };
diff --git a/Engine/Plugins/Developer/Concert/ConcertApp/MultiUserServer/Source/MultiUserServer/Private/Widgets/Clients/Logging/TransportLogDelegates.h b/Engine/Plugins/Developer/Concert/ConcertApp/MultiUserServer/Source/MultiUserServer/Private/Widgets/Clients/Logging/TransportLogDelegates.h
deleted file mode 100644
diff --git a/Engine/Plugins/Developer/Concert/ConcertApp/MultiUserServer/Source/MultiUserServer/Private/Widgets/Clients/Logging/Util/ConcertLogTokenizer.cpp b/Engine/Plugins/Developer/Concert/ConcertApp/MultiUserServer/Source/MultiUserServer/Private/Widgets/Clients/Logging/Util/ConcertLogTokenizer.cpp
--- a/Engine/Plugins/Developer/Concert/ConcertApp/MultiUserServer/Source/MultiUserServer/Private/Widgets/Clients/Logging/Util/ConcertLogTokenizer.cpp
+++ b/Engine/Plugins/Developer/Concert/ConcertApp/MultiUserServer/Source/MultiUserServer/Private/Widgets/Clients/Logging/Util/ConcertLogTokenizer.cpp
@@ -4,17 +4,21 @@
 
 #include "ConcertFrontendUtils.h"
 #include "ConcertTransportEvents.h"
+#include "EndpointToUserNameCache.h"
 #include "MessageTypeUtils.h"
 #include "Math/UnitConversion.h"
 #include "Settings/ConcertTransportLogSettings.h"
 
-FConcertLogTokenizer::FConcertLogTokenizer()
+FConcertLogTokenizer::FConcertLogTokenizer(TSharedRef<FEndpointToUserNameCache> EndpointInfoGetter)
+	: EndpointInfoGetter(EndpointInfoGetter)
 {
 	TokenizerFunctions = {
 			{ FConcertLog::StaticStruct()->FindPropertyByName(GET_MEMBER_NAME_CHECKED(FConcertLog, Timestamp)), [this](const FConcertLog& Log) { return TokenizeTimestamp(Log); } },
 			{ FConcertLog::StaticStruct()->FindPropertyByName(GET_MEMBER_NAME_CHECKED(FConcertLog, MessageTypeName)), [this](const FConcertLog& Log) { return TokenizeMessageTypeName(Log); } },
-			{ FConcertLog::StaticStruct()->FindPropertyByName(GET_MEMBER_NAME_CHECKED(FConcertLog, CustomPayloadUncompressedByteSize)), [this](const FConcertLog& Log) { return TokenizeCustomPayloadUncompressedByteSize(Log); } }
-		};
+			{ FConcertLog::StaticStruct()->FindPropertyByName(GET_MEMBER_NAME_CHECKED(FConcertLog, CustomPayloadUncompressedByteSize)), [this](const FConcertLog& Log) { return TokenizeCustomPayloadUncompressedByteSize(Log); } },
+			{ FConcertLog::StaticStruct()->FindPropertyByName(GET_MEMBER_NAME_CHECKED(FConcertLog, OriginEndpointId)), [this](const FConcertLog& Log) { return TokenizeOriginEndpointId(Log); } },
+			{ FConcertLog::StaticStruct()->FindPropertyByName(GET_MEMBER_NAME_CHECKED(FConcertLog, DestinationEndpointId)), [this](const FConcertLog& Log) { return TokenizeDestinationEndpointId(Log); } }
+	};
 }
 
 FString FConcertLogTokenizer::Tokenize(const FConcertLog& Data, const FProperty& ConcertLogProperty) const
@@ -43,6 +47,16 @@ FString FConcertLogTokenizer::TokenizeCustomPayloadUncompressedByteSize(const FC
 	return FString::Printf(TEXT("%d %s"), DisplayUnit.Value, FUnitConversion::GetUnitDisplayString(DisplayUnit.Units));
 }
 
+FString FConcertLogTokenizer::TokenizeOriginEndpointId(const FConcertLog& Data) const
+{
+	return GetEndpointDisplayString(Data.OriginEndpointId);
+}
+
+FString FConcertLogTokenizer::TokenizeDestinationEndpointId(const FConcertLog& Data) const
+{
+	return GetEndpointDisplayString(Data.DestinationEndpointId);
+}
+
 FString FConcertLogTokenizer::TokenizeUsingPropertyExport(const FConcertLog& Data, const FProperty& ConcertLogProperty) const
 {
 	FString Exported;
@@ -53,3 +67,18 @@ FString FConcertLogTokenizer::TokenizeUsingPropertyExport(const FConcertLog& Dat
 
 	return Exported;
 }
+
+FString FConcertLogTokenizer::GetEndpointDisplayString(const FGuid& EndpointId) const
+{
+	if (EndpointInfoGetter->IsServerEndpoint(EndpointId))
+	{
+		return FString("Server");		
+	}
+
+	if (const TOptional<FConcertClientInfo> ClientInfo = EndpointInfoGetter->GetClientInfo(EndpointId))
+	{
+		return ClientInfo->DisplayName;
+	}
+	
+	return EndpointId.ToString(EGuidFormats::DigitsWithHyphens);
+}
diff --git a/Engine/Plugins/Developer/Concert/ConcertApp/MultiUserServer/Source/MultiUserServer/Private/Widgets/Clients/Logging/Util/ConcertLogTokenizer.h b/Engine/Plugins/Developer/Concert/ConcertApp/MultiUserServer/Source/MultiUserServer/Private/Widgets/Clients/Logging/Util/ConcertLogTokenizer.h
--- a/Engine/Plugins/Developer/Concert/ConcertApp/MultiUserServer/Source/MultiUserServer/Private/Widgets/Clients/Logging/Util/ConcertLogTokenizer.h
+++ b/Engine/Plugins/Developer/Concert/ConcertApp/MultiUserServer/Source/MultiUserServer/Private/Widgets/Clients/Logging/Util/ConcertLogTokenizer.h
@@ -4,6 +4,7 @@
 
 #include "CoreMinimal.h"
 
+class FEndpointToUserNameCache;
 struct FConcertLog;
 
 /** Converts members of FConcertLog into a string. Used e.g. to make search respect the display settings. */
@@ -11,7 +12,7 @@ class FConcertLogTokenizer
 {
 public:
 
-	FConcertLogTokenizer();
+	FConcertLogTokenizer(TSharedRef<FEndpointToUserNameCache> EndpointInfoGetter);
 
 	/** Tokenizes a property of FConcertLog into a string */
 	FString Tokenize(const FConcertLog& Data, const FProperty& ConcertLogProperty) const;
@@ -19,6 +20,8 @@ public:
 	FString TokenizeTimestamp(const FConcertLog& Data) const;
 	FString TokenizeMessageTypeName(const FConcertLog& Data) const;
 	FString TokenizeCustomPayloadUncompressedByteSize(const FConcertLog& Data) const;
+	FString TokenizeOriginEndpointId(const FConcertLog& Data) const;
+	FString TokenizeDestinationEndpointId(const FConcertLog& Data) const;
 	FString TokenizeUsingPropertyExport(const FConcertLog& Data, const FProperty& ConcertLogProperty) const;
 
 private:
@@ -27,4 +30,9 @@ private:
 
 	/** Override functions for tokenizing certain properties */
 	TMap<const FProperty*, FTokenizeFunc> TokenizerFunctions;
+
+	/** Used so we can look up client and server info (even after client has disconnected) */
+	TSharedRef<FEndpointToUserNameCache> EndpointInfoGetter;
+
+	FString GetEndpointDisplayString(const FGuid& EndpointId) const;
 };
diff --git a/Engine/Plugins/Developer/Concert/ConcertApp/MultiUserServer/Source/MultiUserServer/Private/Widgets/Clients/Logging/Util/EndpointToUserNameCache.cpp b/Engine/Plugins/Developer/Concert/ConcertApp/MultiUserServer/Source/MultiUserServer/Private/Widgets/Clients/Logging/Util/EndpointToUserNameCache.cpp
new file mode 100644
--- /dev/null
+++ b/Engine/Plugins/Developer/Concert/ConcertApp/MultiUserServer/Source/MultiUserServer/Private/Widgets/Clients/Logging/Util/EndpointToUserNameCache.cpp
@@ -0,0 +1,142 @@
+// Copyright Epic Games, Inc. All Rights Reserved.
+
+#include "EndpointToUserNameCache.h"
+
+#include "ConcertServerEvents.h"
+#include "ConcertUtil.h"
+#include "IConcertServer.h"
+#include "IConcertSession.h"
+#include "INetworkMessagingExtension.h"
+#include "Algo/AnyOf.h"
+#include "Features/IModularFeatures.h"
+
+FEndpointToUserNameCache::FEndpointToUserNameCache(TSharedRef<IConcertServer> Server)
+	: Server(MoveTemp(Server))
+{
+	ConcertServerEvents::OnLiveSessionCreated().AddRaw(this, &FEndpointToUserNameCache::OnLiveSessionCreated);
+	ConcertServerEvents::OnLiveSessionDestroyed().AddRaw(this, &FEndpointToUserNameCache::OnLiveSessionDestroyed);
+
+	for (TSharedPtr<IConcertServerSession> LiveSession : Server->GetLiveSessions())
+	{
+		RegisterLiveSession(LiveSession.ToSharedRef());
+		for (const FConcertSessionClientInfo& ClientInfo : LiveSession->GetSessionClients())
+		{
+			CacheClientInfo(*LiveSession, ClientInfo);
+		}
+	}
+}
+
+FEndpointToUserNameCache::~FEndpointToUserNameCache()
+{
+	ConcertServerEvents::OnLiveSessionCreated().RemoveAll(this);
+	ConcertServerEvents::OnLiveSessionDestroyed().RemoveAll(this);
+
+	for (TWeakPtr<IConcertServerSession> Session : SubscribedToSessions)
+	{
+		if (TSharedPtr<IConcertServerSession> PinnedSession = Session.Pin())
+		{
+			PinnedSession->OnSessionClientChanged().RemoveAll(this);
+		}
+	}
+}
+
+bool FEndpointToUserNameCache::IsServerEndpoint(const FGuid& EndpointId) const
+{
+	const bool bIsLocalSessionEndpoint = Algo::AnyOf(Server->GetLiveSessions(), [&EndpointId](const TSharedPtr<IConcertServerSession>& Session)
+	{
+		return Session->GetSessionInfo().ServerEndpointId == EndpointId;
+	});
+	return bIsLocalSessionEndpoint || Server->GetServerInfo().AdminEndpointId == EndpointId;
+}
+
+TOptional<FConcertClientInfo> FEndpointToUserNameCache::GetClientInfo(const FGuid& EndpointId) const
+{
+	// Already known or an old endpoint ID that is no longer valid (e.g. because client left session)
+	if (const FNodeEndpointId* CachedNodeId = CachedConcertEndpointToNodeEndpoints.Find(EndpointId))
+	{
+		if (const FConcertClientInfo* ClientData = CachedClientData.Find(*CachedNodeId))
+		{
+			return *ClientData;
+		}
+	}
+	
+	auto LookUpAddress = [this](const FMessageAddress& MessageAddress) -> TOptional<FConcertClientInfo>
+	{
+		const FNodeEndpointId NodeEndpointId = GetNodeIdFromMessagingBackend(MessageAddress); NodeEndpointId.IsValid();
+		if (const FConcertClientInfo* Data = CachedClientData.Find(NodeEndpointId)) 
+		{
+			return *Data;
+		}
+		return {};
+	};
+
+	// Is in any live session?
+	if (const TSharedPtr<IConcertServerSession> Session = ConcertUtil::GetLiveSessionClientConnectedTo(Server.Get(), EndpointId))
+	{
+		const FMessageAddress ConnectedClientAddress = Session->GetClientAddress(EndpointId);
+		return LookUpAddress(ConnectedClientAddress);
+	}
+
+	// Is somebody trying to discover sessions? Data will only be available after this client has joined at least one session
+	if (const FMessageAddress RemoteAdminAddress = Server->GetRemoteAddress(EndpointId); RemoteAdminAddress.IsValid())
+	{
+		return LookUpAddress(RemoteAdminAddress);
+	}
+	return {};
+}
+
+void FEndpointToUserNameCache::OnLiveSessionCreated(bool bSuccess, const IConcertServer& InServer, TSharedRef<IConcertServerSession> InLiveSession)
+{
+	RegisterLiveSession(InLiveSession);
+}
+
+void FEndpointToUserNameCache::OnLiveSessionDestroyed(const IConcertServer& InServer, TSharedRef<IConcertServerSession> InLiveSession)
+{
+	// No need to unsubscribe because the session will be destroyed anyways
+	SubscribedToSessions.Remove(InLiveSession);
+}
+
+void FEndpointToUserNameCache::OnClientInfoChanged(IConcertServerSession& Session, EConcertClientStatus ConnectionStatus, const FConcertSessionClientInfo& ClientInfo)
+{
+	switch (ConnectionStatus)
+	{
+	case EConcertClientStatus::Disconnected:
+		break;
+		
+	case EConcertClientStatus::Connected:
+	case EConcertClientStatus::Updated:
+		CacheClientInfo(Session, ClientInfo);
+		break;
+	default:
+		checkNoEntry();
+	}
+}
+
+void FEndpointToUserNameCache::RegisterLiveSession(const TSharedRef<IConcertServerSession>& InLiveSession)
+{
+	SubscribedToSessions.Add(InLiveSession);
+	InLiveSession->OnSessionClientChanged().AddRaw(this, &FEndpointToUserNameCache::OnClientInfoChanged);
+}
+
+void FEndpointToUserNameCache::CacheClientInfo(const IConcertServerSession& Session, const FConcertSessionClientInfo& ClientInfo)
+{
+	const FMessageAddress Address = Session.GetClientAddress(ClientInfo.ClientEndpointId);
+	if (const FNodeEndpointId NodeEndpointId = GetNodeIdFromMessagingBackend(Address); NodeEndpointId.IsValid())
+	{
+		CachedConcertEndpointToNodeEndpoints.Add(ClientInfo.ClientEndpointId, NodeEndpointId);
+		CachedClientData.Add(NodeEndpointId, ClientInfo.ClientInfo);
+	}
+}
+
+FEndpointToUserNameCache::FNodeEndpointId FEndpointToUserNameCache::GetNodeIdFromMessagingBackend(const FMessageAddress& MessageAddress) const
+{
+	IModularFeatures& ModularFeatures = IModularFeatures::Get();
+	const INetworkMessagingExtension* Extension = ModularFeatures.IsModularFeatureAvailable(INetworkMessagingExtension::ModularFeatureName)
+		? &ModularFeatures.GetModularFeature<INetworkMessagingExtension>(INetworkMessagingExtension::ModularFeatureName)
+		: nullptr;
+	if (Extension)
+	{
+		return Extension->GetNodeIdFromAddress(MessageAddress);
+	}
+	return {};
+}
diff --git a/Engine/Plugins/Developer/Concert/ConcertApp/MultiUserServer/Source/MultiUserServer/Private/Widgets/Clients/Logging/Util/EndpointToUserNameCache.h b/Engine/Plugins/Developer/Concert/ConcertApp/MultiUserServer/Source/MultiUserServer/Private/Widgets/Clients/Logging/Util/EndpointToUserNameCache.h
new file mode 100644
--- /dev/null
+++ b/Engine/Plugins/Developer/Concert/ConcertApp/MultiUserServer/Source/MultiUserServer/Private/Widgets/Clients/Logging/Util/EndpointToUserNameCache.h
@@ -0,0 +1,61 @@
+// Copyright Epic Games, Inc. All Rights Reserved.
+
+#pragma once
+
+#include "CoreMinimal.h"
+#include "ConcertMessageData.h"
+#include "IMessageContext.h"
+
+class IConcertServerSession;
+class IConcertServer;
+
+enum class EConcertClientStatus : uint8;
+
+struct FConcertSessionClientInfo;
+
+/**
+ * Responsible for keeping hold of client info even when the client disconnects.
+ * Client info is removed server-side upon disconnect but logs need their display info.
+ */
+class FEndpointToUserNameCache
+{
+public:
+
+	FEndpointToUserNameCache(TSharedRef<IConcertServer> Server);
+	~FEndpointToUserNameCache();
+
+	bool IsServerEndpoint(const FGuid& EndpointId) const;
+	TOptional<FConcertClientInfo> GetClientInfo(const FGuid& EndpointId) const;
+
+private:
+
+	/** Used to unsubscribe when we're destroyed */
+	TSharedRef<IConcertServer> Server;
+	/** Used to unsubscribe when we're destroyed */
+	TSet<TWeakPtr<IConcertServerSession>> SubscribedToSessions;
+
+	/** ID used by the messaging system - corresponds to an IP address. */
+	using FNodeEndpointId = FGuid;
+	/**
+	 * The client info we're caching.
+	 * 
+	 * Concert may generate multiple endpoint IDs for a single remote machine.
+	 * However, the remote node ID is always unique and is retrieved by querying the UDP backend (see GetNodeIdFromMessagingBackend).
+	 */
+	TMap<FNodeEndpointId, FConcertClientInfo> CachedClientData;
+	/**
+	 * Keeps track of past endpoints IDs that may now no longer be valid.
+	 * Every time a client joins a session, a new endpoint ID is generated for that client. It becomes impossible to
+	 * lookup old concert endpoint IDs without this mapping
+	 */
+	TMap<FGuid, FNodeEndpointId> CachedConcertEndpointToNodeEndpoints;
+
+	void OnLiveSessionCreated(bool bSuccess, const IConcertServer& InServer, TSharedRef<IConcertServerSession> InLiveSession);
+	void OnLiveSessionDestroyed(const IConcertServer& InServer, TSharedRef<IConcertServerSession> InLiveSession);
+	void OnClientInfoChanged(IConcertServerSession& Session, EConcertClientStatus ConnectionStatus, const FConcertSessionClientInfo& ClientInfo);
+
+	void RegisterLiveSession(const TSharedRef<IConcertServerSession>& InLiveSession);
+	void CacheClientInfo(const IConcertServerSession& Session, const FConcertSessionClientInfo& ClientInfo);
+	
+	FNodeEndpointId GetNodeIdFromMessagingBackend(const FMessageAddress& MessageAddress) const;
+};
diff --git a/Engine/Plugins/Developer/Concert/ConcertApp/MultiUserServer/Source/MultiUserServer/Private/Widgets/Clients/SConcertClientsTabView.cpp b/Engine/Plugins/Developer/Concert/ConcertApp/MultiUserServer/Source/MultiUserServer/Private/Widgets/Clients/SConcertClientsTabView.cpp
--- a/Engine/Plugins/Developer/Concert/ConcertApp/MultiUserServer/Source/MultiUserServer/Private/Widgets/Clients/SConcertClientsTabView.cpp
+++ b/Engine/Plugins/Developer/Concert/ConcertApp/MultiUserServer/Source/MultiUserServer/Private/Widgets/Clients/SConcertClientsTabView.cpp
@@ -15,12 +15,13 @@
 #include "Widgets/Clients/Logging/SConcertTransportLog.h"
 
 #include "Framework/Docking/TabManager.h"
+#include "Logging/Util/EndpointToUserNameCache.h"
 #include "Widgets/Docking/SDockTab.h"
 #include "Widgets/Images/SImage.h"
 #include "Widgets/Input/SButton.h"
 #include "Widgets/Text/STextBlock.h"
 
-#define LOCTEXT_NAMESPACE "UnrealMultiUserUI"
+#define LOCTEXT_NAMESPACE "UnrealMultiUserUI.SConcertClientsTabView"
 
 const FName SConcertClientsTabView::ClientBrowserTabId("ClientBrowserTabId");
 const FName SConcertClientsTabView::GlobalLogTabId("GlobalLogTabId");
@@ -29,7 +30,8 @@ void SConcertClientsTabView::Construct(const FArguments& InArgs, FName InStatusB
 {
 	Server = MoveTemp(InServer);
 	LogBuffer = MoveTemp(InLogBuffer);
-	LogTokenizer = MakeShared<FConcertLogTokenizer>();
+	ClientInfoCache = MakeShared<FEndpointToUserNameCache>(Server->GetConcertServer());
+	LogTokenizer = MakeShared<FConcertLogTokenizer>(ClientInfoCache.ToSharedRef());
 	
 	SConcertTabViewWithManagerBase::Construct(
 		SConcertTabViewWithManagerBase::FArguments()
@@ -51,6 +53,7 @@ void SConcertClientsTabView::ShowConnectedClients(const FGuid& SessionId) const
 
 void SConcertClientsTabView::OpenClientLogTab(const FGuid& ClientEndpointId) const
 {
+
 	const FName TabId = *ClientEndpointId.ToString();
 	if (const TSharedPtr<SDockTab> ExistingTab = GetTabManager()->FindExistingLiveTab(FTabId(TabId)))
 	{
@@ -64,7 +67,7 @@ void SConcertClientsTabView::OpenClientLogTab(const FGuid& ClientEndpointId) con
 			.ToolTipText(FText::Format(LOCTEXT("ClientTabTooltipFmt", "Logs all networked requests originating or going to  client {0} (EndpointId = {1})"), FText::FromString(ClientInfo->ClientInfo.DisplayName), FText::FromString(ClientEndpointId.ToString())))
 			.TabRole(PanelTab)
 			[
-				SNew(SConcertTransportLog, LogBuffer.ToSharedRef())
+				SNew(SConcertTransportLog, LogBuffer.ToSharedRef(), ClientInfoCache.ToSharedRef(), LogTokenizer.ToSharedRef())
 				.Filter(UE::MultiUserServer::MakeClientLogFilter(LogTokenizer.ToSharedRef(), ClientEndpointId))
 			]; 
 
@@ -157,7 +160,7 @@ TSharedRef<SDockTab> SConcertClientsTabView::SpawnGlobalLogTab(const FSpawnTabAr
 		.Label(LOCTEXT("GlobalLogTabLabel", "Global Log"))
 		.TabRole(PanelTab)
 		[
-			SNew(SConcertTransportLog, LogBuffer.ToSharedRef())
+			SNew(SConcertTransportLog, LogBuffer.ToSharedRef(), ClientInfoCache.ToSharedRef(), LogTokenizer.ToSharedRef())
 			.Filter(UE::MultiUserServer::MakeGlobalLogFilter(LogTokenizer.ToSharedRef()))
 		]; 
 }
@@ -201,4 +204,18 @@ TSharedRef<SWidget> SConcertClientsTabView::CreateOpenGlobalLogButton() const
 		];
 }
 
+TOptional<FConcertClientInfo> SConcertClientsTabView::GetClientInfo(const FGuid& EndpointId) const
+{
+	if (EndpointId == Server->GetConcertServer()->GetServerInfo().AdminEndpointId)
+	{
+		return {};
+	}
+
+	if (const TOptional<FConcertSessionClientInfo> ClientInfo = ConcertUtil::GetConnectedClientInfo(Server->GetConcertServer().Get(), EndpointId))
+	{
+		return ClientInfo->ClientInfo;
+	}
+	return {};
+}
+
 #undef LOCTEXT_NAMESPACE
diff --git a/Engine/Plugins/Developer/Concert/ConcertApp/MultiUserServer/Source/MultiUserServer/Private/Widgets/Clients/SConcertClientsTabView.h b/Engine/Plugins/Developer/Concert/ConcertApp/MultiUserServer/Source/MultiUserServer/Private/Widgets/Clients/SConcertClientsTabView.h
--- a/Engine/Plugins/Developer/Concert/ConcertApp/MultiUserServer/Source/MultiUserServer/Private/Widgets/Clients/SConcertClientsTabView.h
+++ b/Engine/Plugins/Developer/Concert/ConcertApp/MultiUserServer/Source/MultiUserServer/Private/Widgets/Clients/SConcertClientsTabView.h
@@ -16,8 +16,9 @@ namespace UE
 
 class FConcertClientsTabController;
 class FConcertLogTokenizer;
-class FTabManager;
+class FEndpointToUserNameCache;
 class FGlobalLogSource;
+class FTabManager;
 class IConcertSyncServer;
 class SDockTab;
 class SPromptConcertLoggingEnabled;
@@ -55,6 +56,9 @@ private:
 	TSharedPtr<IConcertSyncServer> Server;
 	/** Buffers all logs globally */
 	TSharedPtr<FGlobalLogSource> LogBuffer;
+
+	/** Caches client info so it remains available even after a client disconnects */
+	TSharedPtr<FEndpointToUserNameCache> ClientInfoCache;
 	/** Used by various systems to convert logs to text */
 	TSharedPtr<FConcertLogTokenizer> LogTokenizer;
 
@@ -70,4 +74,5 @@ private:
 	TSharedRef<SDockTab> SpawnGlobalLogTab(const FSpawnTabArgs& InTabArgs);
 
 	TSharedRef<SWidget> CreateOpenGlobalLogButton() const;
+	TOptional<FConcertClientInfo> GetClientInfo(const FGuid& EndpointId) const;
 };
