Display log ACK state

diff --git a/Engine/Plugins/Developer/Concert/ConcertApp/MultiUserServer/Source/MultiUserServer/Private/ConcertServerStyle.cpp b/Engine/Plugins/Developer/Concert/ConcertApp/MultiUserServer/Source/MultiUserServer/Private/ConcertServerStyle.cpp
--- a/Engine/Plugins/Developer/Concert/ConcertApp/MultiUserServer/Source/MultiUserServer/Private/ConcertServerStyle.cpp
+++ b/Engine/Plugins/Developer/Concert/ConcertApp/MultiUserServer/Source/MultiUserServer/Private/ConcertServerStyle.cpp
@@ -58,18 +58,20 @@ TSharedRef<FSlateStyleSet> FConcertServerStyle::Create()
 	StyleSet->SetContentRoot(FPaths::EngineContentDir() / TEXT("Slate/Starship/Insights"));
 	StyleSet->SetCoreContentRoot(FPaths::EngineContentDir() / TEXT("Slate"));
 	
-	const FVector2D Icon12x12(12.0f, 12.0f); 
 	const FVector2D Icon16x16(16.0f, 16.0f); 
-	const FVector2D Icon20x20(20.0f, 20.0f); 
 	const FVector2D Icon32x32(32.0f, 32.0f); 
-	const FVector2D Icon64x64(64.0f, 64.0f); 
-	
+
+	// Icons
 	StyleSet->Set("Concert.MultiUser", new IMAGE_PLUGIN_BRUSH("Icons/icon_MultiUser_32x", Icon32x32));
 	StyleSet->Set("Concert.SessionContent.ColumnHeader", new IMAGE_PLUGIN_BRUSH_SVG("Icons/Package_16x", Icon16x16));
 	StyleSet->Set("Concert.SessionContent.PackageAdded", new IMAGE_PLUGIN_BRUSH_SVG("Icons/PackageAdded_16x", Icon16x16));
 	StyleSet->Set("Concert.SessionContent.PackageDeleted", new IMAGE_PLUGIN_BRUSH_SVG("Icons/PackageDeleted_16x", Icon16x16));
 	StyleSet->Set("Concert.SessionContent.PackageRenamed", new IMAGE_PLUGIN_BRUSH_SVG("Icons/PackageRenamed_16x", Icon16x16));
 	StyleSet->Set("Concert.SessionContent.PackageSaved", new IMAGE_PLUGIN_BRUSH_SVG("Icons/PackageSaved_16x", Icon16x16));
+	StyleSet->Set("Concert.Ack.Ack", new IMAGE_PLUGIN_BRUSH("Icons/Ack_Ack_16x", Icon16x16));
+	StyleSet->Set("Concert.Ack.Success", new IMAGE_PLUGIN_BRUSH("Icons/Ack_Success_16x", Icon16x16));
+	StyleSet->Set("Concert.Ack.Failure", new IMAGE_PLUGIN_BRUSH("Icons/Ack_Fail_16x", Icon16x16));
+	StyleSet->Set("Concert.Ack.NotNeeded", new IMAGE_PLUGIN_BRUSH("Icons/Ack_NotNeeded_16x", Icon16x16));
 
 	// Clients tab
 	StyleSet->Set("Concert.Clients.DropShadow", new IMAGE_PLUGIN_BOX_BRUSH("ClientThumbnailDropShadow", FMargin(4.0f / 64.0f)));
diff --git a/Engine/Plugins/Developer/Concert/ConcertApp/MultiUserServer/Source/MultiUserServer/Private/Widgets/Clients/ConcertClientsTabController.cpp b/Engine/Plugins/Developer/Concert/ConcertApp/MultiUserServer/Source/MultiUserServer/Private/Widgets/Clients/ConcertClientsTabController.cpp
--- a/Engine/Plugins/Developer/Concert/ConcertApp/MultiUserServer/Source/MultiUserServer/Private/Widgets/Clients/ConcertClientsTabController.cpp
+++ b/Engine/Plugins/Developer/Concert/ConcertApp/MultiUserServer/Source/MultiUserServer/Private/Widgets/Clients/ConcertClientsTabController.cpp
@@ -2,8 +2,10 @@
 
 #include "ConcertClientsTabController.h"
 
+#include "IConcertSyncServer.h"
 #include "SConcertClientsTabView.h"
 #include "Logging/Source/GlobalLogSource.h"
+#include "Logging/Util/LogAckTracker.h"
 #include "Window/ConcertServerTabs.h"
 #include "Window/ConcertServerWindowController.h"
 #include "Widgets/Docking/SDockTab.h"
@@ -12,15 +14,14 @@
 
 namespace UE::MultiUserServer
 {
-	static constexpr uint64 GlobalLogCapacity = 500000;
+	static constexpr uint64 GlobalLogCapacity = 5000000;
 }
 
-FConcertClientsTabController::FConcertClientsTabController()
-	: LogBuffer(MakeShared<FGlobalLogSource>(UE::MultiUserServer::GlobalLogCapacity))
-{}
-
 void FConcertClientsTabController::Init(const FConcertComponentInitParams& Params)
 {
+	LogBuffer = MakeShared<FGlobalLogSource>(UE::MultiUserServer::GlobalLogCapacity);
+	AckTracker = MakeShared<FLogAckTracker>(LogBuffer.ToSharedRef(), Params.Server->GetConcertServer());
+	
 	FGlobalTabmanager::Get()->RegisterTabSpawner(
 				ConcertServerTabs::GetClientsTabID(),
 				FOnSpawnTab::CreateRaw(this, &FConcertClientsTabController::SpawnClientsTab, Params.WindowController->GetRootWindow(), Params.Server)
@@ -44,7 +45,7 @@ TSharedRef<SDockTab> FConcertClientsTabController::SpawnClientsTab(const FSpawnT
 		.TabRole(MajorTab)
 		.CanEverClose(false);
 	DockTab->SetContent(
-		SAssignNew(ClientsView, SConcertClientsTabView, ConcertServerTabs::GetClientsTabID(), Server, LogBuffer)
+		SAssignNew(ClientsView, SConcertClientsTabView, ConcertServerTabs::GetClientsTabID(), Server, LogBuffer.ToSharedRef())
 			.ConstructUnderMajorTab(DockTab)
 			.ConstructUnderWindow(RootWindow)
 		);
diff --git a/Engine/Plugins/Developer/Concert/ConcertApp/MultiUserServer/Source/MultiUserServer/Private/Widgets/Clients/ConcertClientsTabController.h b/Engine/Plugins/Developer/Concert/ConcertApp/MultiUserServer/Source/MultiUserServer/Private/Widgets/Clients/ConcertClientsTabController.h
--- a/Engine/Plugins/Developer/Concert/ConcertApp/MultiUserServer/Source/MultiUserServer/Private/Widgets/Clients/ConcertClientsTabController.h
+++ b/Engine/Plugins/Developer/Concert/ConcertApp/MultiUserServer/Source/MultiUserServer/Private/Widgets/Clients/ConcertClientsTabController.h
@@ -5,6 +5,7 @@
 #include "CoreMinimal.h"
 #include "Window/IConcertComponent.h"
 
+class FLogAckTracker;
 class FGlobalLogSource;
 class SConcertClientsTabView;
 class SWindow;
@@ -12,8 +13,6 @@ class SWindow;
 class FConcertClientsTabController : public IConcertComponent, public TSharedFromThis<FConcertClientsTabController>
 {
 public:
-	
-	FConcertClientsTabController();
 
 	//~ Begin IConcertComponent Interface
 	virtual void Init(const FConcertComponentInitParams& Params) override;
@@ -25,7 +24,9 @@ public:
 private:
 
 	/** Buffers generated logs up to a limit (and overrides oldest logs when buffer is full) */
-	TSharedRef<FGlobalLogSource> LogBuffer;
+	TSharedPtr<FGlobalLogSource> LogBuffer;
+	/** Marks logs as acked as the ACKs come in. */
+	TSharedPtr<FLogAckTracker> AckTracker;
 	
 	/** Manages the sub-tabs */
 	TSharedPtr<SConcertClientsTabView> ClientsView;
diff --git a/Engine/Plugins/Developer/Concert/ConcertApp/MultiUserServer/Source/MultiUserServer/Private/Widgets/Clients/Logging/ConcertLogEntry.h b/Engine/Plugins/Developer/Concert/ConcertApp/MultiUserServer/Source/MultiUserServer/Private/Widgets/Clients/Logging/ConcertLogEntry.h
--- a/Engine/Plugins/Developer/Concert/ConcertApp/MultiUserServer/Source/MultiUserServer/Private/Widgets/Clients/Logging/ConcertLogEntry.h
+++ b/Engine/Plugins/Developer/Concert/ConcertApp/MultiUserServer/Source/MultiUserServer/Private/Widgets/Clients/Logging/ConcertLogEntry.h
@@ -4,9 +4,48 @@
 
 #include "CoreMinimal.h"
 #include "ConcertTransportEvents.h"
+#include "ConcertLogEntry.generated.h"
 
 using FConcertLogID = uint64;
 
+UENUM()
+enum class EConcertLogAckState : uint8
+{
+	/** This message does not require any acks */
+	NotNeeded,
+	/** This is an ack */
+	Ack,
+	/** Still waiting on the ack */
+	InProgress,
+	/** The ack was received */
+	AckReceived,
+	/** Never received any acks - set after a certain timeout. */
+	AckFailure
+};
+
+/** Additional properties we display in the UI but that are not in FConcertLog live here. */
+USTRUCT()
+struct FConcertLogMetadata
+{
+	GENERATED_BODY()
+
+	/** The ack state of this log */
+	UPROPERTY()
+	EConcertLogAckState AckState = EConcertLogAckState::NotNeeded;
+
+	/**
+	 * Message ID of this log's ACK. If set, implies this log is not an ACK.
+	 * Valid when AckState == EConcertLogAckState::AckReceived.
+	 */
+	TOptional<FGuid> AckingMessageId;
+
+	/**
+	 * The Message IDs of the logs we've ACKed. If set, implies this log is an ACK.
+	 * Valid when AckState == EConcertLogAckState::Ack.
+	 */
+	TOptional<TSet<FGuid>> AckedMessageId;
+};
+
 struct FConcertLogEntry
 {
 	/** Unique log ID. Log IDs grow sequentially. */
@@ -14,4 +53,7 @@ struct FConcertLogEntry
 
 	/** The log we're describing */
 	FConcertLog Log;
+
+	/** Additional data to display in the UI */
+	FConcertLogMetadata LogMetaData;
 };
diff --git a/Engine/Plugins/Developer/Concert/ConcertApp/MultiUserServer/Source/MultiUserServer/Private/Widgets/Clients/Logging/Filter/ConcertFrontendLogFilter.h b/Engine/Plugins/Developer/Concert/ConcertApp/MultiUserServer/Source/MultiUserServer/Private/Widgets/Clients/Logging/Filter/ConcertFrontendLogFilter.h
--- a/Engine/Plugins/Developer/Concert/ConcertApp/MultiUserServer/Source/MultiUserServer/Private/Widgets/Clients/Logging/Filter/ConcertFrontendLogFilter.h
+++ b/Engine/Plugins/Developer/Concert/ConcertApp/MultiUserServer/Source/MultiUserServer/Private/Widgets/Clients/Logging/Filter/ConcertFrontendLogFilter.h
@@ -24,7 +24,7 @@ class TConcertFrontendLogFilterAggregate : public FConcertFrontendLogFilter
 public:
 
 	//~ Begin FConcertLogFilter Interface
-	virtual bool PassesFilter(const FConcertLog& InItem) const final override { return Implementation.PassesFilter(InItem); }
+	virtual bool PassesFilter(const FConcertLogEntry& InItem) const final override { return Implementation.PassesFilter(InItem); }
 	//~ End FConcertLogFilter Interface
 	
 	//~ Begin FConcertFrontendLogFilter Interface
diff --git a/Engine/Plugins/Developer/Concert/ConcertApp/MultiUserServer/Source/MultiUserServer/Private/Widgets/Clients/Logging/Filter/ConcertFrontendLogFilter_Ack.h b/Engine/Plugins/Developer/Concert/ConcertApp/MultiUserServer/Source/MultiUserServer/Private/Widgets/Clients/Logging/Filter/ConcertFrontendLogFilter_Ack.h
new file mode 100644
--- /dev/null
+++ b/Engine/Plugins/Developer/Concert/ConcertApp/MultiUserServer/Source/MultiUserServer/Private/Widgets/Clients/Logging/Filter/ConcertFrontendLogFilter_Ack.h
@@ -0,0 +1,58 @@
+// Copyright Epic Games, Inc. All Rights Reserved.
+
+#pragma once
+
+#include "CoreMinimal.h"
+#include "ConcertFrontendLogFilter_BaseSetSelection.h"
+#include "Widgets/Clients/Logging/ConcertLogEntry.h"
+
+namespace UE::MultiUserServer::Filters
+{
+	/** Allows only messages that have a selected ACK state */
+	class FConcertLogFilter_Ack : public TConcertLogFilter_BaseSetSelection<FConcertLogFilter_Ack, EConcertLogAckState>
+	{
+	public:
+
+		static TSet<EConcertLogAckState> GetAllOptions()
+		{
+			return {
+				EConcertLogAckState::Ack,
+				EConcertLogAckState::AckFailure,
+				EConcertLogAckState::AckReceived,
+				EConcertLogAckState::InProgress,
+				EConcertLogAckState::NotNeeded
+			};
+		}
+		
+		static FString GetOptionDisplayString(const EConcertLogAckState& Item)
+		{
+			switch(Item)
+			{
+				case EConcertLogAckState::Ack: return TEXT("Ack");
+				case EConcertLogAckState::AckFailure: return TEXT("Ack Failure");
+				case EConcertLogAckState::AckReceived: return TEXT("Ack Received");
+				case EConcertLogAckState::InProgress: return TEXT("In Progress");
+				case EConcertLogAckState::NotNeeded: return TEXT("Not needed");
+				default: checkNoEntry(); return {}; 
+			}
+		}
+		
+		//~ Begin FConcertLogFilter Interface
+		virtual bool PassesFilter(const FConcertLogEntry& InItem) const override
+		{
+			return IsItemAllowed(InItem.LogMetaData.AckState);
+		}
+		//~ End FConcertLogFilter Interface
+	};
+
+	class FConcertFrontendLogFilter_Ack : public TConcertFrontendLogFilter_BaseSetSelection<FConcertLogFilter_Ack>
+	{
+	public:
+
+		FConcertFrontendLogFilter_Ack()
+			: TConcertFrontendLogFilter_BaseSetSelection<FConcertLogFilter_Ack>(NSLOCTEXT("UnrealMultiUserUI.Filter.Ack", "Name", "ACKs"))
+		{}
+	};
+}
+
+
diff --git a/Engine/Plugins/Developer/Concert/ConcertApp/MultiUserServer/Source/MultiUserServer/Private/Widgets/Clients/Logging/Filter/ConcertFrontendLogFilter_BaseSetSelection.h b/Engine/Plugins/Developer/Concert/ConcertApp/MultiUserServer/Source/MultiUserServer/Private/Widgets/Clients/Logging/Filter/ConcertFrontendLogFilter_BaseSetSelection.h
new file mode 100644
--- /dev/null
+++ b/Engine/Plugins/Developer/Concert/ConcertApp/MultiUserServer/Source/MultiUserServer/Private/Widgets/Clients/Logging/Filter/ConcertFrontendLogFilter_BaseSetSelection.h
@@ -0,0 +1,203 @@
+// Copyright Epic Games, Inc. All Rights Reserved.
+
+#pragma once
+
+#include "CoreMinimal.h"
+#include "ConcertFrontendLogFilter.h"
+#include "ConcertLogFilter.h"
+
+#include "Framework/MultiBox/MultiBoxBuilder.h"
+#include "Widgets/Input/SComboButton.h"
+#include "Widgets/SBoxPanel.h"
+
+#define LOCTEXT_NAMESPACE "UnrealMultiUserUI.BaseSetSelection"
+
+namespace UE::MultiUserServer::Filters
+{
+	/**
+	 * Helper implementation for filters that allow filter based on a finte set of items.
+	 * Subclasses must implement:
+	 * - static TSet<TSetType> GetAllOptions();
+	 * - static FString GetOptionDisplayString(const TSetType&);
+	 */
+	template<typename TRealFilterImpl, typename TSetType>
+	class TConcertLogFilter_BaseSetSelection : public FConcertLogFilter
+	{
+	public:
+
+		using TItemType = TSetType;
+
+		TConcertLogFilter_BaseSetSelection()
+			: AllowedItems(TRealFilterImpl::GetAllOptions())
+		{}
+
+		void AllowAll()
+		{
+			TSet<TSetType> Allowed = TRealFilterImpl::GetAllOptions();
+			if (Allowed.Num() != AllowedItems.Num())
+			{
+				AllowedItems = MoveTemp(Allowed);
+				OnChanged().Broadcast();
+			}
+		}
+		void DisallowAll()
+		{
+			if (AllowedItems.Num() > 0)
+			{
+				AllowedItems.Reset();
+				OnChanged().Broadcast();
+			}
+		}
+		void ToggleAll(const TSet<TSetType>& ToToggle)
+		{
+			for (const TSetType Item : ToToggle)
+			{
+				if (IsItemAllowed(Item))
+				{
+					DisallowItem(Item);
+				}
+				else
+				{
+					AllowItem(Item);
+				}
+			}
+
+			if (ToToggle.Num() > 0)
+			{
+				OnChanged().Broadcast();
+			}
+		}
+	
+		void AllowItem(const TSetType& MessageTypeName)
+		{
+			if (!AllowedItems.Contains(MessageTypeName))
+			{
+				AllowedItems.Add(MessageTypeName);
+				OnChanged().Broadcast();
+			}
+		}
+		void DisallowItem(const TSetType& MessageTypeName)
+		{
+			if (AllowedItems.Contains(MessageTypeName))
+			{
+				AllowedItems.Remove(MessageTypeName);
+				OnChanged().Broadcast();
+			}
+		}
+	
+		bool IsItemAllowed(const TSetType& MessageTypeName) const
+		{
+			return AllowedItems.Contains(MessageTypeName);
+		}
+		bool AreAllAllowed() const
+		{
+			return AllowedItems.Num() == TRealFilterImpl::GetAllOptions().Num();
+		}
+		uint8 GetNumSelected() const { return AllowedItems.Num(); }
+		
+	private:
+	
+		TSet<TSetType> AllowedItems;
+	};
+
+	template<typename TRealFilterImpl>
+	class TConcertFrontendLogFilter_BaseSetSelection : public TConcertFrontendLogFilterAggregate<TRealFilterImpl>
+	{
+		using TItemType = typename TRealFilterImpl::TItemType;
+	public:
+
+		TConcertFrontendLogFilter_BaseSetSelection(FText FilterName)
+		{
+			this->ChildSlot = SNew(SHorizontalBox)
+				.ToolTipText(LOCTEXT("BaseSetSelection.ToolTipText", "Select a list of allowed items\nHint: Type in menu to search"))
+
+				+SHorizontalBox::Slot()
+				.AutoWidth()
+				.VAlign(VAlign_Center)
+				[
+					SNew(STextBlock)
+					.Text(FilterName)
+				]
+
+				+SHorizontalBox::Slot()
+				.AutoWidth()
+				.Padding(2, 0, 0, 0)
+				.VAlign(VAlign_Center)
+				[
+					SNew(SComboButton)
+					.OnGetMenuContent_Raw(this, &TConcertFrontendLogFilter_BaseSetSelection::MakeSelectionMenu)
+					.ButtonContent()
+					[
+						SNew(STextBlock)
+						.Text_Lambda([this]()
+						{
+							return this->Implementation.AreAllAllowed()
+								? LOCTEXT("BaseSetSelection.Selection.All", "All")
+								: FText::FromString(FString::FromInt(this->Implementation.GetNumSelected()));
+						})
+					]
+				];
+		}
+
+	private:
+
+		TSharedRef<SWidget> MakeSelectionMenu()
+		{
+			FMenuBuilder MenuBuilder(false, nullptr);
+
+			MenuBuilder.AddMenuEntry(
+				LOCTEXT("BaseSetSelection.SelectAll.", "Select all"),
+				LOCTEXT("BaseSetSelection.SelectAll.Tooltip", "Allows all items"),
+				FSlateIcon(),
+				FUIAction(
+					FExecuteAction::CreateLambda([this](){ this->Implementation.AllowAll(); }),
+					FCanExecuteAction::CreateLambda([] { return true; }),
+					FIsActionChecked()),
+				NAME_None,
+				EUserInterfaceActionType::Button
+				);
+			MenuBuilder.AddMenuEntry(
+				LOCTEXT("BaseSetSelection.DeselectAll.", "Deselect all"),
+				LOCTEXT("BaseSetSelection.DeelectAll.Tooltip", "Disallows all items"),
+				FSlateIcon(),
+				FUIAction(
+					FExecuteAction::CreateLambda([this](){ this->Implementation.DisallowAll(); }),
+					FCanExecuteAction::CreateLambda([] { return true; }),
+					FIsActionChecked()),
+				NAME_None,
+				EUserInterfaceActionType::Button
+				);
+	
+			MenuBuilder.AddSeparator();
+	
+			for (const TItemType& Item : TRealFilterImpl::GetAllOptions())
+			{
+				MenuBuilder.AddMenuEntry(
+					FText::FromString(TRealFilterImpl::GetOptionDisplayString(Item)),
+					FText::GetEmpty(),
+					FSlateIcon(),
+					FUIAction(
+						FExecuteAction::CreateLambda([this, Item]()
+						{
+							if (this->Implementation.IsItemAllowed(Item))
+							{
+								this->Implementation.DisallowItem(Item);
+							}
+							else
+							{
+								this->Implementation.AllowItem(Item);
+							}
+						}),
+						FCanExecuteAction::CreateLambda([] { return true; }),
+						FIsActionChecked::CreateLambda([this, Item]() { return this->Implementation.IsItemAllowed(Item); })),
+					NAME_None,
+					EUserInterfaceActionType::ToggleButton
+				);
+			}
+			return MenuBuilder.MakeWidget();
+		}
+	};
+
+}
+
+#undef LOCTEXT_NAMESPACE
\ No newline at end of file
diff --git a/Engine/Plugins/Developer/Concert/ConcertApp/MultiUserServer/Source/MultiUserServer/Private/Widgets/Clients/Logging/Filter/ConcertFrontendLogFilter_Client.cpp b/Engine/Plugins/Developer/Concert/ConcertApp/MultiUserServer/Source/MultiUserServer/Private/Widgets/Clients/Logging/Filter/ConcertFrontendLogFilter_Client.cpp
--- a/Engine/Plugins/Developer/Concert/ConcertApp/MultiUserServer/Source/MultiUserServer/Private/Widgets/Clients/Logging/Filter/ConcertFrontendLogFilter_Client.cpp
+++ b/Engine/Plugins/Developer/Concert/ConcertApp/MultiUserServer/Source/MultiUserServer/Private/Widgets/Clients/Logging/Filter/ConcertFrontendLogFilter_Client.cpp
@@ -3,11 +3,12 @@
 #include "ConcertFrontendLogFilter_Client.h"
 
 #include "ConcertTransportEvents.h"
+#include "Widgets/Clients/Logging/ConcertLogEntry.h"
 
-bool FConcertLogFilter_Client::PassesFilter(const FConcertLog& InItem) const
+bool FConcertLogFilter_Client::PassesFilter(const FConcertLogEntry& InItem) const
 {
-	return AllowedClientEndpointIds.Contains(InItem.OriginEndpointId)
-		|| AllowedClientEndpointIds.Contains(InItem.DestinationEndpointId);
+	return AllowedClientEndpointIds.Contains(InItem.Log.OriginEndpointId)
+		|| AllowedClientEndpointIds.Contains(InItem.Log.DestinationEndpointId);
 }
 
 void FConcertLogFilter_Client::AllowOnly(const FGuid& ClientEndpointId)
diff --git a/Engine/Plugins/Developer/Concert/ConcertApp/MultiUserServer/Source/MultiUserServer/Private/Widgets/Clients/Logging/Filter/ConcertFrontendLogFilter_Client.h b/Engine/Plugins/Developer/Concert/ConcertApp/MultiUserServer/Source/MultiUserServer/Private/Widgets/Clients/Logging/Filter/ConcertFrontendLogFilter_Client.h
--- a/Engine/Plugins/Developer/Concert/ConcertApp/MultiUserServer/Source/MultiUserServer/Private/Widgets/Clients/Logging/Filter/ConcertFrontendLogFilter_Client.h
+++ b/Engine/Plugins/Developer/Concert/ConcertApp/MultiUserServer/Source/MultiUserServer/Private/Widgets/Clients/Logging/Filter/ConcertFrontendLogFilter_Client.h
@@ -16,7 +16,7 @@ public:
 	{}
 
 	//~ Begin FConcertLogFilter Interface
-	virtual bool PassesFilter(const FConcertLog& InItem) const override;
+	virtual bool PassesFilter(const FConcertLogEntry& InItem) const override;
 	//~ End FConcertLogFilter Interface
 
 	void AllowOnly(const FGuid& ClientEndpointId);
diff --git a/Engine/Plugins/Developer/Concert/ConcertApp/MultiUserServer/Source/MultiUserServer/Private/Widgets/Clients/Logging/Filter/ConcertFrontendLogFilter_MessageAction.cpp b/Engine/Plugins/Developer/Concert/ConcertApp/MultiUserServer/Source/MultiUserServer/Private/Widgets/Clients/Logging/Filter/ConcertFrontendLogFilter_MessageAction.cpp
deleted file mode 100644
diff --git a/Engine/Plugins/Developer/Concert/ConcertApp/MultiUserServer/Source/MultiUserServer/Private/Widgets/Clients/Logging/Filter/ConcertFrontendLogFilter_MessageAction.h b/Engine/Plugins/Developer/Concert/ConcertApp/MultiUserServer/Source/MultiUserServer/Private/Widgets/Clients/Logging/Filter/ConcertFrontendLogFilter_MessageAction.h
--- a/Engine/Plugins/Developer/Concert/ConcertApp/MultiUserServer/Source/MultiUserServer/Private/Widgets/Clients/Logging/Filter/ConcertFrontendLogFilter_MessageAction.h
+++ b/Engine/Plugins/Developer/Concert/ConcertApp/MultiUserServer/Source/MultiUserServer/Private/Widgets/Clients/Logging/Filter/ConcertFrontendLogFilter_MessageAction.h
@@ -3,42 +3,43 @@
 #pragma once
 
 #include "CoreMinimal.h"
+#include "ConcertFrontendLogFilter_BaseSetSelection.h"
+#include "Widgets/Clients/Logging/ConcertLogEntry.h"
 #include "Widgets/Clients/Logging/Filter/ConcertFrontendLogFilter.h"
+#include "Widgets/Clients/Logging/Util/MessageActionUtils.h"
 
-/** Allows only the selected messages types */
-class FConcertLogFilter_MessageAction : public FConcertLogFilter
+namespace UE::MultiUserServer::Filters
 {
-public:
-
-	FConcertLogFilter_MessageAction();
-
-	//~ Begin FConcertLogFilter Interface
-	virtual bool PassesFilter(const FConcertLog& InItem) const override;
-	//~ End FConcertLogFilter Interface
-
-	void AllowAll();
-	void DisallowAll();
-	void ToggleAll(const TSet<FName>& ToToggle);
-	
-	void AllowMessageAction(FName MessageTypeName);
-	void DisallowMessageAction(FName MessageTypeName);
-	
-	bool IsMessageActionAllowed(FName MessageTypeName) const;
-	bool AreAllAllowed() const;
-	uint8 GetNumSelected() const;
-	
-private:
-	
-	TSet<FName> AllowedMessageActionNames;
-};
-
-class FConcertFrontendLogFilter_MessageAction : public TConcertFrontendLogFilterAggregate<FConcertLogFilter_MessageAction>
-{
-public:
-
-	FConcertFrontendLogFilter_MessageAction();
-
-private:
+	/** Allows only the selected messages types */
+	class FConcertLogFilter_MessageAction : public TConcertLogFilter_BaseSetSelection<FConcertLogFilter_MessageAction, FName>
+	{
+	public:
+
+		static TSet<FName> GetAllOptions()
+		{
+			return MessageActionUtils::GetAllMessageActionNames();
+		}
+		
+		static FString GetOptionDisplayString(const FName& Item)
+		{
+			return MessageActionUtils::GetActionDisplayString(Item);
+		}
+
+		//~ Begin FConcertLogFilter Interface
+		virtual bool PassesFilter(const FConcertLogEntry& InItem) const override
+		{
+			return IsItemAllowed(MessageActionUtils::ConvertActionToName(InItem.Log.MessageAction));
+		}
+		//~ End FConcertLogFilter Interface
+	};
+
+	class FConcertFrontendLogFilter_MessageAction : public TConcertFrontendLogFilter_BaseSetSelection<FConcertLogFilter_MessageAction>
+	{
+	public:
+
+		FConcertFrontendLogFilter_MessageAction()
+			: TConcertFrontendLogFilter_BaseSetSelection<FConcertLogFilter_MessageAction>(NSLOCTEXT("UnrealMultiUserUI.Filter.MessageAction", "Name", "Actions"))
+		{}
+	};
+}
 
-	TSharedRef<SWidget> MakeSelectionMenu();
-};
diff --git a/Engine/Plugins/Developer/Concert/ConcertApp/MultiUserServer/Source/MultiUserServer/Private/Widgets/Clients/Logging/Filter/ConcertFrontendLogFilter_MessageType.cpp b/Engine/Plugins/Developer/Concert/ConcertApp/MultiUserServer/Source/MultiUserServer/Private/Widgets/Clients/Logging/Filter/ConcertFrontendLogFilter_MessageType.cpp
deleted file mode 100644
diff --git a/Engine/Plugins/Developer/Concert/ConcertApp/MultiUserServer/Source/MultiUserServer/Private/Widgets/Clients/Logging/Filter/ConcertFrontendLogFilter_MessageType.h b/Engine/Plugins/Developer/Concert/ConcertApp/MultiUserServer/Source/MultiUserServer/Private/Widgets/Clients/Logging/Filter/ConcertFrontendLogFilter_MessageType.h
--- a/Engine/Plugins/Developer/Concert/ConcertApp/MultiUserServer/Source/MultiUserServer/Private/Widgets/Clients/Logging/Filter/ConcertFrontendLogFilter_MessageType.h
+++ b/Engine/Plugins/Developer/Concert/ConcertApp/MultiUserServer/Source/MultiUserServer/Private/Widgets/Clients/Logging/Filter/ConcertFrontendLogFilter_MessageType.h
@@ -3,42 +3,46 @@
 #pragma once
 
 #include "CoreMinimal.h"
-#include "Widgets/Clients/Logging/Filter/ConcertFrontendLogFilter.h"
+#include "ConcertFrontendLogFilter_BaseSetSelection.h"
+#include "Widgets/Clients/Logging/ConcertLogEntry.h"
+#include "Widgets/Clients/Logging/Util/MessageTypeUtils.h"
 
-/** Allows only the selected messages types */
-class FConcertLogFilter_MessageType : public FConcertLogFilter
+namespace UE::MultiUserServer::Filters
 {
-public:
-
-	FConcertLogFilter_MessageType();
-
-	//~ Begin FConcertLogFilter Interface
-	virtual bool PassesFilter(const FConcertLog& InItem) const override;
-	//~ End FConcertLogFilter Interface
-
-	void AllowAll();
-	void DisallowAll();
-	void ToggleAll(const TSet<FName>& ToToggle);
-	
-	void AllowMessageType(FName MessageTypeName);
-	void DisallowMessageType(FName MessageTypeName);
+	/** Allows only the selected messages types */
+	class FConcertLogFilter_MessageType : public TConcertLogFilter_BaseSetSelection<FConcertLogFilter_MessageType, FName>
+	{
+	public:
+
+		static TSet<FName> GetAllOptions()
+		{
+			return MessageTypeUtils::GetAllMessageTypeNames();
+		}
+		
+		static FString GetOptionDisplayString(const FName& Item)
+		{
+			return MessageTypeUtils::SanitizeMessageTypeName(Item);
+		}
+		
+		//~ Begin FConcertLogFilter Interface
+		virtual bool PassesFilter(const FConcertLogEntry& InItem) const override
+		{
+			return IsItemAllowed(InItem.Log.MessageTypeName);
+		}
+		//~ End FConcertLogFilter Interface
 	
-	bool IsMessageTypeAllowed(FName MessageTypeName) const;
-	bool AreAllAllowed() const;
-	uint8 GetNumSelected() const;
+	private:
 	
-private:
-	
-	TSet<FName> AllowedMessageTypeNames;
-};
-
-class FConcertFrontendLogFilter_MessageType : public TConcertFrontendLogFilterAggregate<FConcertLogFilter_MessageType>
-{
-public:
+		TSet<FName> AllowedMessageTypeNames;
+	};
 
-	FConcertFrontendLogFilter_MessageType();
+	class FConcertFrontendLogFilter_MessageType : public TConcertFrontendLogFilter_BaseSetSelection<FConcertLogFilter_MessageType>
+	{
+	public:
 
-private:
+		FConcertFrontendLogFilter_MessageType()
+			: TConcertFrontendLogFilter_BaseSetSelection<FConcertLogFilter_MessageType>(NSLOCTEXT("UnrealMultiUserUI.Filter.MessageType", "Name", "Message Type"))
+		{}
+	};
+}
 
-	TSharedRef<SWidget> MakeSelectionMenu();
-};
diff --git a/Engine/Plugins/Developer/Concert/ConcertApp/MultiUserServer/Source/MultiUserServer/Private/Widgets/Clients/Logging/Filter/ConcertFrontendLogFilter_Size.cpp b/Engine/Plugins/Developer/Concert/ConcertApp/MultiUserServer/Source/MultiUserServer/Private/Widgets/Clients/Logging/Filter/ConcertFrontendLogFilter_Size.cpp
--- a/Engine/Plugins/Developer/Concert/ConcertApp/MultiUserServer/Source/MultiUserServer/Private/Widgets/Clients/Logging/Filter/ConcertFrontendLogFilter_Size.cpp
+++ b/Engine/Plugins/Developer/Concert/ConcertApp/MultiUserServer/Source/MultiUserServer/Private/Widgets/Clients/Logging/Filter/ConcertFrontendLogFilter_Size.cpp
@@ -8,24 +8,25 @@
 
 #include "Framework/MultiBox/MultiBoxBuilder.h"
 #include "Widgets/SBoxPanel.h"
+#include "Widgets/Clients/Logging/ConcertLogEntry.h"
 #include "Widgets/Input/SButton.h"
 #include "Widgets/Input/SNumericEntryBox.h"
 #include "Widgets/Text/STextBlock.h"
 
 #define LOCTEXT_NAMESPACE "UnrealMultiUserUI.FConcertLogFilter_Size"
 
-bool FConcertLogFilter_Size::PassesFilter(const FConcertLog& InItem) const
+bool FConcertLogFilter_Size::PassesFilter(const FConcertLogEntry& InItem) const
 {
-	check(InItem.CustomPayloadUncompressedByteSize >= 0);
+	check(InItem.Log.CustomPayloadUncompressedByteSize >= 0);
 	// Note: This only filters activities events - they all use custom events.
 	// The filter's default value is to show everything 0 <= x so it will show sync events as well.
 	const uint32 ComparisionSizeInBytes = FUnitConversion::Convert(SizeInBytes, DataUnit, EUnit::Bytes);
 	switch (FilterMode)
 	{
 	case ESizeFilterMode::LessThanOrEqual:
-		return ComparisionSizeInBytes >= static_cast<uint32>(InItem.CustomPayloadUncompressedByteSize);
+		return ComparisionSizeInBytes >= static_cast<uint32>(InItem.Log.CustomPayloadUncompressedByteSize);
 	case ESizeFilterMode::BiggerThanOrEqual:
-		return ComparisionSizeInBytes <= static_cast<uint32>(InItem.CustomPayloadUncompressedByteSize);
+		return ComparisionSizeInBytes <= static_cast<uint32>(InItem.Log.CustomPayloadUncompressedByteSize);
 	default:
 		checkNoEntry();
 		return false;
diff --git a/Engine/Plugins/Developer/Concert/ConcertApp/MultiUserServer/Source/MultiUserServer/Private/Widgets/Clients/Logging/Filter/ConcertFrontendLogFilter_Size.h b/Engine/Plugins/Developer/Concert/ConcertApp/MultiUserServer/Source/MultiUserServer/Private/Widgets/Clients/Logging/Filter/ConcertFrontendLogFilter_Size.h
--- a/Engine/Plugins/Developer/Concert/ConcertApp/MultiUserServer/Source/MultiUserServer/Private/Widgets/Clients/Logging/Filter/ConcertFrontendLogFilter_Size.h
+++ b/Engine/Plugins/Developer/Concert/ConcertApp/MultiUserServer/Source/MultiUserServer/Private/Widgets/Clients/Logging/Filter/ConcertFrontendLogFilter_Size.h
@@ -20,7 +20,7 @@ class FConcertLogFilter_Size : public FConcertLogFilter
 public:
 	
 	//~ Begin FConcertLogFilter Interface
-	virtual bool PassesFilter(const FConcertLog& InItem) const override;
+	virtual bool PassesFilter(const FConcertLogEntry& InItem) const override;
 	//~ End FConcertLogFilter Interface
 
 	void AdvanceFilterMode();
diff --git a/Engine/Plugins/Developer/Concert/ConcertApp/MultiUserServer/Source/MultiUserServer/Private/Widgets/Clients/Logging/Filter/ConcertFrontendLogFilter_TextSearch.cpp b/Engine/Plugins/Developer/Concert/ConcertApp/MultiUserServer/Source/MultiUserServer/Private/Widgets/Clients/Logging/Filter/ConcertFrontendLogFilter_TextSearch.cpp
--- a/Engine/Plugins/Developer/Concert/ConcertApp/MultiUserServer/Source/MultiUserServer/Private/Widgets/Clients/Logging/Filter/ConcertFrontendLogFilter_TextSearch.cpp
+++ b/Engine/Plugins/Developer/Concert/ConcertApp/MultiUserServer/Source/MultiUserServer/Private/Widgets/Clients/Logging/Filter/ConcertFrontendLogFilter_TextSearch.cpp
@@ -3,10 +3,11 @@
 #include "ConcertFrontendLogFilter_TextSearch.h"
 
 #include "ConcertTransportEvents.h"
+#include "Widgets/Clients/Logging/ConcertLogEntry.h"
 #include "Widgets/Clients/Logging/Util/ConcertLogTokenizer.h"
 
 FConcertLogFilter_TextSearch::FConcertLogFilter_TextSearch(TSharedRef<FConcertLogTokenizer> Tokenizer)
-	: TextFilter(TTextFilter<const FConcertLog&>::FItemToStringArray::CreateRaw(this, &FConcertLogFilter_TextSearch::GenerateSearchTerms))
+	: TextFilter(TTextFilter<const FConcertLogEntry&>::FItemToStringArray::CreateRaw(this, &FConcertLogFilter_TextSearch::GenerateSearchTerms))
 	, Tokenizer(MoveTemp(Tokenizer))
 {
 	TextFilter.OnChanged().AddLambda([this]
@@ -15,11 +16,11 @@ FConcertLogFilter_TextSearch::FConcertLogFilter_TextSearch(TSharedRef<FConcertLo
 	});
 }
 
-void FConcertLogFilter_TextSearch::GenerateSearchTerms(const FConcertLog& InItem, TArray<FString>& OutTerms) const
+void FConcertLogFilter_TextSearch::GenerateSearchTerms(const FConcertLogEntry& InItem, TArray<FString>& OutTerms) const
 {
 	for (TFieldIterator<const FProperty> PropertyIt(FConcertLog::StaticStruct()); PropertyIt; ++PropertyIt)
 	{
-		OutTerms.Add(Tokenizer->Tokenize(InItem, **PropertyIt));
+		OutTerms.Add(Tokenizer->Tokenize(InItem.Log, **PropertyIt));
 	}
 }
 
diff --git a/Engine/Plugins/Developer/Concert/ConcertApp/MultiUserServer/Source/MultiUserServer/Private/Widgets/Clients/Logging/Filter/ConcertFrontendLogFilter_TextSearch.h b/Engine/Plugins/Developer/Concert/ConcertApp/MultiUserServer/Source/MultiUserServer/Private/Widgets/Clients/Logging/Filter/ConcertFrontendLogFilter_TextSearch.h
--- a/Engine/Plugins/Developer/Concert/ConcertApp/MultiUserServer/Source/MultiUserServer/Private/Widgets/Clients/Logging/Filter/ConcertFrontendLogFilter_TextSearch.h
+++ b/Engine/Plugins/Developer/Concert/ConcertApp/MultiUserServer/Source/MultiUserServer/Private/Widgets/Clients/Logging/Filter/ConcertFrontendLogFilter_TextSearch.h
@@ -17,7 +17,7 @@ public:
 	FConcertLogFilter_TextSearch(TSharedRef<FConcertLogTokenizer> Tokenizer);
 	
 	//~ Begin FConcertLogFilter Interface
-	virtual bool PassesFilter(const FConcertLog& InItem) const override { return TextFilter.PassesFilter(InItem); }
+	virtual bool PassesFilter(const FConcertLogEntry& InItem) const override { return TextFilter.PassesFilter(InItem); }
 	//~ End FConcertLogFilter Interface
 	
 	void SetRawFilterText(const FText& InFilterText) { TextFilter.SetRawFilterText(InFilterText); }
@@ -25,12 +25,12 @@ public:
 private:
 
 	/** Does the actual string search */
-	TTextFilter<const FConcertLog&> TextFilter;
+	TTextFilter<const FConcertLogEntry&> TextFilter;
 	/** Helps in converting FConcertLog members into search terms */
 	TSharedRef<FConcertLogTokenizer> Tokenizer;
 
 	/** Parses InItem into a bunch of strings that can be searched */
-	void GenerateSearchTerms(const FConcertLog& InItem, TArray<FString>& OutTerms) const;
+	void GenerateSearchTerms(const FConcertLogEntry& InItem, TArray<FString>& OutTerms) const;
 };
 
 /** Creates a search bar */
diff --git a/Engine/Plugins/Developer/Concert/ConcertApp/MultiUserServer/Source/MultiUserServer/Private/Widgets/Clients/Logging/Filter/ConcertFrontendLogFilter_Time.cpp b/Engine/Plugins/Developer/Concert/ConcertApp/MultiUserServer/Source/MultiUserServer/Private/Widgets/Clients/Logging/Filter/ConcertFrontendLogFilter_Time.cpp
--- a/Engine/Plugins/Developer/Concert/ConcertApp/MultiUserServer/Source/MultiUserServer/Private/Widgets/Clients/Logging/Filter/ConcertFrontendLogFilter_Time.cpp
+++ b/Engine/Plugins/Developer/Concert/ConcertApp/MultiUserServer/Source/MultiUserServer/Private/Widgets/Clients/Logging/Filter/ConcertFrontendLogFilter_Time.cpp
@@ -4,6 +4,7 @@
 
 #include "ConcertFrontendUtils.h"
 #include "ConcertTransportEvents.h"
+#include "Widgets/Clients/Logging/ConcertLogEntry.h"
 
 #include "Framework/MultiBox/MultiBoxBuilder.h"
 #include "Widgets/Input/SComboButton.h"
@@ -29,12 +30,12 @@ FDateTime FConcertLogFilter_Time::MakeResetTime() const
 	return FilterMode == ETimeFilter::AllowAfter ? FDateTime() : FDateTime::MaxValue();
 }
 
-bool FConcertLogFilter_Time::PassesFilter(const FConcertLog& InItem) const
+bool FConcertLogFilter_Time::PassesFilter(const FConcertLogEntry& InItem) const
 {
 	switch (FilterMode)
 	{
-	case ETimeFilter::AllowAfter: return InItem.Timestamp >= Time;
-	case ETimeFilter::AllowBefore: return InItem.Timestamp <= Time;
+	case ETimeFilter::AllowAfter: return InItem.Log.Timestamp >= Time;
+	case ETimeFilter::AllowBefore: return InItem.Log.Timestamp <= Time;
 	default:
 		checkNoEntry();
 		return true;
diff --git a/Engine/Plugins/Developer/Concert/ConcertApp/MultiUserServer/Source/MultiUserServer/Private/Widgets/Clients/Logging/Filter/ConcertFrontendLogFilter_Time.h b/Engine/Plugins/Developer/Concert/ConcertApp/MultiUserServer/Source/MultiUserServer/Private/Widgets/Clients/Logging/Filter/ConcertFrontendLogFilter_Time.h
--- a/Engine/Plugins/Developer/Concert/ConcertApp/MultiUserServer/Source/MultiUserServer/Private/Widgets/Clients/Logging/Filter/ConcertFrontendLogFilter_Time.h
+++ b/Engine/Plugins/Developer/Concert/ConcertApp/MultiUserServer/Source/MultiUserServer/Private/Widgets/Clients/Logging/Filter/ConcertFrontendLogFilter_Time.h
@@ -23,7 +23,7 @@ public:
 	void ResetToInfiniteTime();
 
 	//~ Begin FConcertLogFilter Interface
-	virtual bool PassesFilter(const FConcertLog& InItem) const override;
+	virtual bool PassesFilter(const FConcertLogEntry& InItem) const override;
 	//~ End FConcertLogFilter Interfac
 
 	ETimeFilter GetFilterMode() const { return FilterMode; }
diff --git a/Engine/Plugins/Developer/Concert/ConcertApp/MultiUserServer/Source/MultiUserServer/Private/Widgets/Clients/Logging/Filter/ConcertLogFilter.h b/Engine/Plugins/Developer/Concert/ConcertApp/MultiUserServer/Source/MultiUserServer/Private/Widgets/Clients/Logging/Filter/ConcertLogFilter.h
--- a/Engine/Plugins/Developer/Concert/ConcertApp/MultiUserServer/Source/MultiUserServer/Private/Widgets/Clients/Logging/Filter/ConcertLogFilter.h
+++ b/Engine/Plugins/Developer/Concert/ConcertApp/MultiUserServer/Source/MultiUserServer/Private/Widgets/Clients/Logging/Filter/ConcertLogFilter.h
@@ -5,11 +5,11 @@
 #include "CoreMinimal.h"
 #include "Misc/IFilter.h"
 
-struct FConcertLog;
+struct FConcertLogEntry;
 
 /** Base filter */
 class FConcertLogFilter :
-	public IFilter<const FConcertLog&>,
+	public IFilter<const FConcertLogEntry&>,
 	public TSharedFromThis<FConcertLogFilter>,
 	// We do not need to copy filters so let's avoid it by accident; some constructors pass this pointer to callbacks which become stale upon copying
 	public FNoncopyable
@@ -17,7 +17,7 @@ class FConcertLogFilter :
 public:
 
 	//~ Begin IFilter Interface
-	DECLARE_DERIVED_EVENT( FFrontendFilter, IFilter<const FConcertLog&>::FChangedEvent, FChangedEvent );
+	DECLARE_DERIVED_EVENT( FFrontendFilter, IFilter<const FConcertLogEntry&>::FChangedEvent, FChangedEvent );
 	virtual FChangedEvent& OnChanged() override { return ChangedEvent; }
 	//~ End IFilter Interface
 
diff --git a/Engine/Plugins/Developer/Concert/ConcertApp/MultiUserServer/Source/MultiUserServer/Private/Widgets/Clients/Logging/Filter/ConcertLogFilter_FrontendRoot.cpp b/Engine/Plugins/Developer/Concert/ConcertApp/MultiUserServer/Source/MultiUserServer/Private/Widgets/Clients/Logging/Filter/ConcertLogFilter_FrontendRoot.cpp
--- a/Engine/Plugins/Developer/Concert/ConcertApp/MultiUserServer/Source/MultiUserServer/Private/Widgets/Clients/Logging/Filter/ConcertLogFilter_FrontendRoot.cpp
+++ b/Engine/Plugins/Developer/Concert/ConcertApp/MultiUserServer/Source/MultiUserServer/Private/Widgets/Clients/Logging/Filter/ConcertLogFilter_FrontendRoot.cpp
@@ -3,6 +3,7 @@
 #include "ConcertLogFilter_FrontendRoot.h"
 
 #include "ConcertFrontendLogFilter.h"
+#include "ConcertFrontendLogFilter_Ack.h"
 #include "ConcertFrontendLogFilter_Client.h"
 #include "ConcertFrontendLogFilter_MessageAction.h"
 #include "ConcertFrontendLogFilter_MessageType.h"
@@ -55,7 +56,7 @@ TSharedRef<SWidget> FConcertLogFilter_FrontendRoot::BuildFilterWidgets() const
 		];
 }
 
-bool FConcertLogFilter_FrontendRoot::PassesFilter(const FConcertLog& InItem) const
+bool FConcertLogFilter_FrontendRoot::PassesFilter(const FConcertLogEntry& InItem) const
 {
 	return Algo::AllOf(
 		AllFrontendFilters,
@@ -91,11 +92,12 @@ namespace UE::MultiUserServer
 		static TArray<TSharedRef<FConcertFrontendLogFilter>> CreateCommonFilters()
 		{
 			return {
-				MakeShared<FConcertFrontendLogFilter_MessageAction>(),
-				MakeShared<FConcertFrontendLogFilter_MessageType>(),
+				MakeShared<Filters::FConcertFrontendLogFilter_MessageAction>(),
+				MakeShared<Filters::FConcertFrontendLogFilter_MessageType>(),
 				MakeShared<FConcertFrontendLogFilter_Time>(ETimeFilter::AllowAfter),
 				MakeShared<FConcertFrontendLogFilter_Time>(ETimeFilter::AllowBefore),
-				MakeShared<FConcertFrontendLogFilter_Size>()
+				MakeShared<FConcertFrontendLogFilter_Size>(),
+				MakeShared<Filters::FConcertFrontendLogFilter_Ack>()
 			};
 		}
 	}
diff --git a/Engine/Plugins/Developer/Concert/ConcertApp/MultiUserServer/Source/MultiUserServer/Private/Widgets/Clients/Logging/Filter/ConcertLogFilter_FrontendRoot.h b/Engine/Plugins/Developer/Concert/ConcertApp/MultiUserServer/Source/MultiUserServer/Private/Widgets/Clients/Logging/Filter/ConcertLogFilter_FrontendRoot.h
--- a/Engine/Plugins/Developer/Concert/ConcertApp/MultiUserServer/Source/MultiUserServer/Private/Widgets/Clients/Logging/Filter/ConcertLogFilter_FrontendRoot.h
+++ b/Engine/Plugins/Developer/Concert/ConcertApp/MultiUserServer/Source/MultiUserServer/Private/Widgets/Clients/Logging/Filter/ConcertLogFilter_FrontendRoot.h
@@ -23,7 +23,7 @@ public:
 	TSharedRef<SWidget> BuildFilterWidgets() const;
 	
 	//~ Begin IFilter Interface
-	virtual bool PassesFilter(const FConcertLog& InItem) const override;
+	virtual bool PassesFilter(const FConcertLogEntry& InItem) const override;
 	//~ End IFilter Interface
 
 	FORCEINLINE const TSharedRef<FConcertFrontendLogFilter_TextSearch>& GetTextSearchFilter() const { return TextSearchFilter; }
diff --git a/Engine/Plugins/Developer/Concert/ConcertApp/MultiUserServer/Source/MultiUserServer/Private/Widgets/Clients/Logging/Filter/FilteredConcertLogList.cpp b/Engine/Plugins/Developer/Concert/ConcertApp/MultiUserServer/Source/MultiUserServer/Private/Widgets/Clients/Logging/Filter/FilteredConcertLogList.cpp
--- a/Engine/Plugins/Developer/Concert/ConcertApp/MultiUserServer/Source/MultiUserServer/Private/Widgets/Clients/Logging/Filter/FilteredConcertLogList.cpp
+++ b/Engine/Plugins/Developer/Concert/ConcertApp/MultiUserServer/Source/MultiUserServer/Private/Widgets/Clients/Logging/Filter/FilteredConcertLogList.cpp
@@ -46,7 +46,7 @@ void FFilteredConcertLogList::OnLowestLogEntryChanged(FConcertLogID NewLowestVal
 
 void FFilteredConcertLogList::OnNewLogEntryAdded(const TSharedRef<FConcertLogEntry>& NewLogEntry)
 {
-	if (!Filter || Filter->PassesFilter(NewLogEntry->Log))
+	if (!Filter || Filter->PassesFilter(*NewLogEntry))
 	{
 		FilteredLogs.Emplace(NewLogEntry);
 		OnLogListChanged().Broadcast(GetFilteredLogs());
@@ -58,7 +58,7 @@ void FFilteredConcertLogList::RebuildFilteredResult()
 	FilteredLogs.Reset();
 	LogSource->ForEachLog([this](const TSharedPtr<FConcertLogEntry>& LogEntry)
 	{
-		if (!Filter || Filter->PassesFilter(LogEntry->Log))
+		if (!Filter || Filter->PassesFilter(*LogEntry))
 		{
 			FilteredLogs.Emplace(LogEntry);
 		}
@@ -77,6 +77,16 @@ FPagedFilteredConcertLogList::FPagedFilteredConcertLogList(TSharedRef<IConcertLo
 	CheckAndConditionallyPopulatePage();
 }
 
+TOptional<FPagedFilteredConcertLogList::FPageCount> FPagedFilteredConcertLogList::GetPageOf(const size_t ItemIndex) const
+{
+	const FPageCount Page = ItemIndex / GetLogsPerPage();
+	if (Page < GetNumPages())
+	{
+		return Page;
+	}
+	return {};
+}
+
 void FPagedFilteredConcertLogList::SetLogsPerPage(uint16 NewLogsPerPage)
 {
 	if (LogsPerPage != NewLogsPerPage && ensure(NewLogsPerPage > 0))
@@ -141,19 +151,18 @@ void FPagedFilteredConcertLogList::CheckAndConditionallyPopulatePage()
 		++PageViewIndex;
 	}, FMath::Min(PageView.Num(), GetFilteredLogs().Num()));
 
-	size_t StartIndex, LastIndex;
-	Tie(StartIndex, LastIndex) = GetLogIndicesForPage();
+	size_t StartIndex, ItemsOnPage;
+	Tie(StartIndex, ItemsOnPage) = GetLogIndicesForPage();
 
 	// Add items that are new in GetFilteredLogs()
-	bChanged |= PageView.Num() <= LastIndex; 
-	for (size_t i = PageView.Num(); i <= LastIndex && GetFilteredLogs().IsValidIndex(i); ++i)
+	bChanged |= PageView.Num() < ItemsOnPage; 
+	for (size_t i = PageView.Num(); i < ItemsOnPage && GetFilteredLogs().IsValidIndex(i); ++i)
 	{
 		PageView.Emplace(GetFilteredLogs()[i]);
 	}
 
 	// Remove (possibly duplicate) items that exceed the expected page length 
-	const FLogsPerPageCount LogsOnPage = FMath::Min(LastIndex, LastIndex - StartIndex + 1);
-	const size_t NumToRemove = PageView.Num() - LogsOnPage;
+	const size_t NumToRemove = PageView.Num() - ItemsOnPage;
 	bChanged |= NumToRemove > 0;
 	for (size_t i = 0; i < NumToRemove; ++i)
 	{
@@ -168,9 +177,8 @@ void FPagedFilteredConcertLogList::CheckAndConditionallyPopulatePage()
 
 void FPagedFilteredConcertLogList::ForEachLogIndexOnPage(TFunctionRef<void(size_t Index)> Callback, FLogsPerPageCount MaxItems)
 {
-	size_t StartIndex, LastIndex;
-	Tie(StartIndex, LastIndex) = GetLogIndicesForPage();
-	const FLogsPerPageCount NumItems = FMath::Min(LastIndex, LastIndex - StartIndex + 1);
+	size_t StartIndex, NumItems;
+	Tie(StartIndex, NumItems) = GetLogIndicesForPage();
 	for (size_t i = 0; i < NumItems && i < MaxItems; ++i)
 	{
 		Callback(StartIndex + i);
@@ -180,12 +188,12 @@ void FPagedFilteredConcertLogList::ForEachLogIndexOnPage(TFunctionRef<void(size_
 TTuple<size_t, size_t> FPagedFilteredConcertLogList::GetLogIndicesForPage() const
 {
 	const size_t StartIndex = CurrentPageIndex * LogsPerPage;
-	const size_t LastIndex = FMath::Min<size_t>(
+	const size_t NumberItems = FMath::Min<size_t>(
 		// Used when page is full
-		StartIndex + LogsPerPage - 1,
+		LogsPerPage,
 		// Used when last page is not full
-		GetFilteredLogs().Num() == 0 ? 0 : GetFilteredLogs().Num() - 1
+		GetFilteredLogs().Num() - StartIndex
 		);
 	check(GetFilteredLogs().Num() == 0 || GetFilteredLogs().IsValidIndex(StartIndex));
-	return { StartIndex, LastIndex };
+	return { StartIndex, NumberItems };
 }
diff --git a/Engine/Plugins/Developer/Concert/ConcertApp/MultiUserServer/Source/MultiUserServer/Private/Widgets/Clients/Logging/Filter/FilteredConcertLogList.h b/Engine/Plugins/Developer/Concert/ConcertApp/MultiUserServer/Source/MultiUserServer/Private/Widgets/Clients/Logging/Filter/FilteredConcertLogList.h
--- a/Engine/Plugins/Developer/Concert/ConcertApp/MultiUserServer/Source/MultiUserServer/Private/Widgets/Clients/Logging/Filter/FilteredConcertLogList.h
+++ b/Engine/Plugins/Developer/Concert/ConcertApp/MultiUserServer/Source/MultiUserServer/Private/Widgets/Clients/Logging/Filter/FilteredConcertLogList.h
@@ -57,11 +57,12 @@ public:
 
 	FPageCount GetNumPages() const { const FPageCount PageCount = GetFilteredLogs().Num() / LogsPerPage; return GetFilteredLogs().Num() % LogsPerPage == 0 ? PageCount : PageCount + 1;}
 	FPageCount GetCurrentPage() const { return CurrentPageIndex; }
-	uint16 GetLogsPerPage() const { return LogsPerPage; }
+	TOptional<FPageCount> GetPageOf(const size_t ItemIndex) const;
+	FLogsPerPageCount GetLogsPerPage() const { return LogsPerPage; }
 	const TArray<TSharedPtr<FConcertLogEntry>>& GetPageView() const { return PageView; }
 	
 	void SetLogsPerPage(FLogsPerPageCount NewLogsPerPage);
-	void SetPage(const uint32 PageIndex);
+	void SetPage(const FPageCount PageIndex);
 
 	FOnLogListChanged& OnPageViewChanged() { return OnPageViewChangedEvent; }
 	
diff --git a/Engine/Plugins/Developer/Concert/ConcertApp/MultiUserServer/Source/MultiUserServer/Private/Widgets/Clients/Logging/SConcertTransportLog.cpp b/Engine/Plugins/Developer/Concert/ConcertApp/MultiUserServer/Source/MultiUserServer/Private/Widgets/Clients/Logging/SConcertTransportLog.cpp
--- a/Engine/Plugins/Developer/Concert/ConcertApp/MultiUserServer/Source/MultiUserServer/Private/Widgets/Clients/Logging/SConcertTransportLog.cpp
+++ b/Engine/Plugins/Developer/Concert/ConcertApp/MultiUserServer/Source/MultiUserServer/Private/Widgets/Clients/Logging/SConcertTransportLog.cpp
@@ -93,7 +93,7 @@ TSharedRef<SWidget> SConcertTransportLog::CreateTableView()
 	return SAssignNew(LogView, SListView<TSharedPtr<FConcertLogEntry>>)
 		.ListItemsSource(&PagedLogList->GetPageView())
 		.OnGenerateRow(this, &SConcertTransportLog::OnGenerateActivityRowWidget)
-		.SelectionMode(ESelectionMode::None)
+		.SelectionMode(ESelectionMode::Multi)
 		.HeaderRow(CreateHeaderRow());
 }
 
@@ -117,30 +117,51 @@ TSharedRef<SHeaderRow> SConcertTransportLog::CreateHeaderRow()
 			.ShouldGenerateWidget(true)
 			.ToolTipText(LOCTEXT("AvatarColumnToolTipText", "The colour of the avatar is affected by log"))
 		;
-
-	const TMap<FName, FString> ColumnNameOverrides = {
-		{ GET_MEMBER_NAME_CHECKED(FConcertLog, CustomPayloadUncompressedByteSize), TEXT("Size") },
-		{ GET_MEMBER_NAME_CHECKED(FConcertLog, OriginEndpointId), TEXT("Origin") },
-		{ GET_MEMBER_NAME_CHECKED(FConcertLog, DestinationEndpointId), TEXT("Destination") }
-	};
-	for (TFieldIterator<const FProperty> PropertyIt(FConcertLog::StaticStruct()); PropertyIt; ++PropertyIt)
+	
+	ForEachPropertyColumn([this](const FProperty& Property, FName ColumnId)
 	{
-		if (!PropertyIt->HasAnyPropertyFlags(CPF_Transient))
+		const TMap<FName, FString> ColumnNameOverrides = {
+			{ GET_MEMBER_NAME_CHECKED(FConcertLog, MessageAction), TEXT("Action") },
+			{ GET_MEMBER_NAME_CHECKED(FConcertLog, MessageOrderIndex), TEXT("Order Index") },
+			{ GET_MEMBER_NAME_CHECKED(FConcertLog, ChannelId), TEXT("Channel") },
+			{ GET_MEMBER_NAME_CHECKED(FConcertLog, CustomPayloadUncompressedByteSize), TEXT("Size") },
+			{ GET_MEMBER_NAME_CHECKED(FConcertLog, OriginEndpointId), TEXT("Origin") },
+			{ GET_MEMBER_NAME_CHECKED(FConcertLog, DestinationEndpointId), TEXT("Destination") },
+			{ GET_MEMBER_NAME_CHECKED(FConcertLogMetadata, AckState), TEXT("Ack") }
+		};
+		const TMap<FName, float> ManuallySizeColumns = {
+			{ GET_MEMBER_NAME_CHECKED(FConcertLog, Frame), 60 },
+			{ GET_MEMBER_NAME_CHECKED(FConcertLog, MessageId), 280 },
+			{ GET_MEMBER_NAME_CHECKED(FConcertLog, MessageOrderIndex), 60 },
+			{ GET_MEMBER_NAME_CHECKED(FConcertLog, ChannelId), 60 },
+			{ GET_MEMBER_NAME_CHECKED(FConcertLog, Timestamp), 160 },
+			{ GET_MEMBER_NAME_CHECKED(FConcertLog, MessageAction), 80 },
+			{ GET_MEMBER_NAME_CHECKED(FConcertLog, MessageTypeName), 200 },
+			{ GET_MEMBER_NAME_CHECKED(FConcertLog, CustomPayloadTypename), 400 },
+			{ GET_MEMBER_NAME_CHECKED(FConcertLog, MessageAction), 80 },
+			{ GET_MEMBER_NAME_CHECKED(FConcertLog, CustomPayloadUncompressedByteSize), 100 },
+			{ GET_MEMBER_NAME_CHECKED(FConcertLogMetadata, AckState), 45 }
+		};
+		
+		const FString* LabelOverride = ColumnNameOverrides.Find(ColumnId);
+		const FString PropertyName = LabelOverride ? *LabelOverride : Property.GetAuthoredName();
+		SHeaderRow::FColumn::FArguments Args = SHeaderRow::FColumn::FArguments()
+			.ColumnId(ColumnId)
+			.DefaultLabel(FText::FromString(PropertyName))
+			.HAlignCell(HAlign_Center)
+			// Add option to hide
+			.OnGetMenuContent_Lambda([this, ColumnId]()
+			{
+				return UE::ConcertSharedSlate::MakeHideColumnContextMenu(HeaderRow.ToSharedRef(), ColumnId);
+			});
+		
+		if (const float* ManualSize = ManuallySizeColumns.Find(ColumnId))
 		{
-			const FName ColumnId = PropertyIt->GetFName();
-			const FString* LabelOverride = ColumnNameOverrides.Find(ColumnId);
-			const FString PropertyName = LabelOverride ? *LabelOverride : PropertyIt->GetAuthoredName();
-			HeaderRow->AddColumn(SHeaderRow::FColumn::FArguments()
-				.ColumnId(ColumnId)
-				.DefaultLabel(FText::FromString(PropertyName))
-				.HAlignCell(HAlign_Center)
-				// Add option to hide
-				.OnGetMenuContent_Lambda([this, ColumnId]()
-				{
-					return UE::ConcertSharedSlate::MakeHideColumnContextMenu(HeaderRow.ToSharedRef(), ColumnId);
-				}));
+			Args.ManualWidth(*ManualSize);
 		}
-	}
+		
+		HeaderRow->AddColumn(Args);
+	});
 
 	TGuardValue<bool> DoNotSave(bIsUpdatingColumnVisibility, true);
 	RestoreDefaultColumnVisiblities();
@@ -157,25 +178,42 @@ TSharedRef<ITableRow> SConcertTransportLog::OnGenerateActivityRowWidget(TSharedP
 		: DestinationInfo.IsSet() ? DestinationInfo->AvatarColor : FLinearColor::Black;
 	
 	return SNew(SConcertTransportLogRow, Item, OwnerTable, LogTokenizer.ToSharedRef(), HighlightText.ToSharedRef())
-		.AvatarColor(AvatarColor);
+		.AvatarColor(AvatarColor)
+		.ScrollToAckLog(this, &SConcertTransportLog::ScrollToAckLog)
+		.ScrollToAckedLog(this, &SConcertTransportLog::ScrollToAckedLog);
 }
 
-void SConcertTransportLog::RestoreDefaultColumnVisiblities()
+void SConcertTransportLog::ForEachPropertyColumn(TFunctionRef<void(const FProperty& ColumnPropertyId, FName ColumnId)> Callback)
 {
-	const TSet<FName> HiddenByDefault = {
-		GET_MEMBER_NAME_CHECKED(FConcertLog, Frame),
-		GET_MEMBER_NAME_CHECKED(FConcertLog, MessageId),
-		GET_MEMBER_NAME_CHECKED(FConcertLog, MessageOrderIndex),
-		GET_MEMBER_NAME_CHECKED(FConcertLog, ChannelId),
-		GET_MEMBER_NAME_CHECKED(FConcertLog, CustomPayloadTypename),
-		GET_MEMBER_NAME_CHECKED(FConcertLog, StringPayload)
-	};
-	for (const FName ColumnID : HiddenByDefault)
+	for (TFieldIterator<const FProperty> PropertyIt(FConcertLog::StaticStruct()); PropertyIt; ++PropertyIt)
+	{
+		if (!PropertyIt->HasAnyPropertyFlags(CPF_Transient))
+		{
+			Callback(**PropertyIt, PropertyIt->GetFName());
+		}
+	}
+	for (TFieldIterator<const FProperty> PropertyIt(FConcertLogMetadata::StaticStruct()); PropertyIt; ++PropertyIt)
 	{
-		HeaderRow->SetShowGeneratedColumn(ColumnID, false);
+		Callback(**PropertyIt, PropertyIt->GetFName());
 	}
 }
 
+void SConcertTransportLog::RestoreDefaultColumnVisiblities()
+{
+	ForEachPropertyColumn([this](const FProperty& Property, FName ColumnId)
+	{
+		const TSet<FName> HiddenByDefault = {
+			GET_MEMBER_NAME_CHECKED(FConcertLog, Frame),
+			GET_MEMBER_NAME_CHECKED(FConcertLog, MessageId),
+			GET_MEMBER_NAME_CHECKED(FConcertLog, MessageOrderIndex),
+			GET_MEMBER_NAME_CHECKED(FConcertLog, ChannelId),
+			GET_MEMBER_NAME_CHECKED(FConcertLog, CustomPayloadTypename),
+			GET_MEMBER_NAME_CHECKED(FConcertLog, StringPayload)
+		};
+		HeaderRow->SetShowGeneratedColumn(ColumnId, !HiddenByDefault.Contains(ColumnId));
+	});
+}
+
 void SConcertTransportLog::ExtendViewOptions(FMenuBuilder& MenuBuilder)
 {
 	MenuBuilder.AddMenuEntry(
@@ -209,6 +247,38 @@ void SConcertTransportLog::ExtendViewOptions(FMenuBuilder& MenuBuilder)
 		);
 
 	MenuBuilder.AddSeparator();
+	MenuBuilder.AddMenuEntry(
+			LOCTEXT("SelectAll", "Show all"),
+			FText::GetEmpty(),
+			FSlateIcon(),
+			FUIAction(
+				FExecuteAction::CreateLambda([this]()
+				{
+					ForEachPropertyColumn([this](const FProperty& Property, FName ColumnId)
+					{
+						HeaderRow->SetShowGeneratedColumn(ColumnId, true);
+					});
+				}),
+				FCanExecuteAction::CreateLambda([] { return true; })),
+			NAME_None,
+			EUserInterfaceActionType::Button
+		);
+	MenuBuilder.AddMenuEntry(
+			LOCTEXT("SelectAll", "Hide all"),
+			FText::GetEmpty(),
+			FSlateIcon(),
+			FUIAction(
+				FExecuteAction::CreateLambda([this]()
+				{
+					ForEachPropertyColumn([this](const FProperty& Property, FName ColumnId)
+					{
+						HeaderRow->SetShowGeneratedColumn(ColumnId, false);
+					});
+				}),
+				FCanExecuteAction::CreateLambda([] { return true; })),
+			NAME_None,
+			EUserInterfaceActionType::Button
+		);
 	MenuBuilder.AddMenuEntry(
 			LOCTEXT("RestoreDefaultColumnVisibility", "Restore columns visibility"),
 			FText::GetEmpty(),
@@ -219,6 +289,7 @@ void SConcertTransportLog::ExtendViewOptions(FMenuBuilder& MenuBuilder)
 			NAME_None,
 			EUserInterfaceActionType::Button
 		);
+	MenuBuilder.AddSeparator();
 	UE::ConcertSharedSlate::AddEntriesForShowingHiddenRows(HeaderRow.ToSharedRef(), MenuBuilder);
 }
 
@@ -274,4 +345,36 @@ void SConcertTransportLog::OnConcertLoggingEnabledChanged(bool bNewEnabled)
 	}
 }
 
+void SConcertTransportLog::ScrollToAckLog(const FGuid& MessageId) const
+{
+	const int32 Index = PagedLogList->GetFilteredLogs().IndexOfByPredicate([MessageId](const TSharedPtr<FConcertLogEntry>& Entry)
+	{
+		return Entry->Log.MessageId == MessageId && Entry->Log.MessageAction == EConcertLogMessageAction::Process;
+	});
+	ScrollToLog(Index);
+}
+
+void SConcertTransportLog::ScrollToAckedLog(const FGuid& MessageId) const
+{
+	const int32 Index = PagedLogList->GetFilteredLogs().IndexOfByPredicate([MessageId](const TSharedPtr<FConcertLogEntry>& Entry)
+	{
+		return Entry->Log.MessageId == MessageId && Entry->Log.MessageAction == EConcertLogMessageAction::Send;
+	});
+	ScrollToLog(Index);
+}
+
+void SConcertTransportLog::ScrollToLog(const int32 LogIndex) const
+{
+	if (const TOptional<FPagedFilteredConcertLogList::FPageCount> PageIndex = PagedLogList->GetPageOf(LogIndex); PageIndex && PagedLogList->GetFilteredLogs().IsValidIndex(LogIndex))
+	{
+		PagedLogList->SetPage(*PageIndex);
+		LogView->RequestNavigateToItem(PagedLogList->GetFilteredLogs()[LogIndex]);
+		LogView->SetSelection(PagedLogList->GetFilteredLogs()[LogIndex]);
+	}
+	else
+	{
+		// TODO: Tell the user. Ideally the button calling would be disabled but that requires an optmised search algorithm for checking whether an item is currently filtered out or not.
+	}
+}
+
 #undef LOCTEXT_NAMESPACE
diff --git a/Engine/Plugins/Developer/Concert/ConcertApp/MultiUserServer/Source/MultiUserServer/Private/Widgets/Clients/Logging/SConcertTransportLog.h b/Engine/Plugins/Developer/Concert/ConcertApp/MultiUserServer/Source/MultiUserServer/Private/Widgets/Clients/Logging/SConcertTransportLog.h
--- a/Engine/Plugins/Developer/Concert/ConcertApp/MultiUserServer/Source/MultiUserServer/Private/Widgets/Clients/Logging/SConcertTransportLog.h
+++ b/Engine/Plugins/Developer/Concert/ConcertApp/MultiUserServer/Source/MultiUserServer/Private/Widgets/Clients/Logging/SConcertTransportLog.h
@@ -73,6 +73,8 @@ private:
 	TSharedRef<SHeaderRow> CreateHeaderRow();
 	TSharedRef<ITableRow> OnGenerateActivityRowWidget(TSharedPtr<FConcertLogEntry> Item, const TSharedRef<STableViewBase>& OwnerTable) const;
 
+	static void ForEachPropertyColumn(TFunctionRef<void(const FProperty& ColumnPropertyId, FName ColumnId)> Callback);
+
 	void RestoreDefaultColumnVisiblities();
 	void ExtendViewOptions(FMenuBuilder& MenuBuilder);
 	void OnFilterMenuChecked();
@@ -81,5 +83,8 @@ private:
 	void OnColumnVisibilitySettingsChanged(const FColumnVisibilitySnapshot& ColumnSnapshot);
 	
 	void OnConcertLoggingEnabledChanged(bool bNewEnabled);
-};
 
+	void ScrollToAckLog(const FGuid& MessageId) const;
+	void ScrollToAckedLog(const FGuid& MessageId) const;
+	void ScrollToLog(const int32 LogIndex) const;
+};
diff --git a/Engine/Plugins/Developer/Concert/ConcertApp/MultiUserServer/Source/MultiUserServer/Private/Widgets/Clients/Logging/SConcertTransportLogRow.cpp b/Engine/Plugins/Developer/Concert/ConcertApp/MultiUserServer/Source/MultiUserServer/Private/Widgets/Clients/Logging/SConcertTransportLogRow.cpp
--- a/Engine/Plugins/Developer/Concert/ConcertApp/MultiUserServer/Source/MultiUserServer/Private/Widgets/Clients/Logging/SConcertTransportLogRow.cpp
+++ b/Engine/Plugins/Developer/Concert/ConcertApp/MultiUserServer/Source/MultiUserServer/Private/Widgets/Clients/Logging/SConcertTransportLogRow.cpp
@@ -2,13 +2,21 @@
 
 #include "SConcertTransportLogRow.h"
 
+#include "ConcertServerStyle.h"
 #include "Util/ConcertLogTokenizer.h"
 #include "SConcertTransportLog.h"
 #include "Session/Activity/PredefinedActivityColumns.h"
 
 #include "Widgets/Colors/SColorBlock.h"
+#include "Widgets/Input/SButton.h"
+#include "Widgets/Input/SComboButton.h"
+#include "Widgets/SOverlay.h"
+#include "Widgets/Images/SThrobber.h"
+#include "Widgets/Layout/SScaleBox.h"
 #include "Widgets/Text/STextBlock.h"
 
+#define LOCTEXT_NAMESPACE "UnrealMultiUserUI.SConcertTransportLogRow"
+
 void SConcertTransportLogRow::Construct(const FArguments& InArgs, TSharedPtr<FConcertLogEntry> InLogEntry, const TSharedRef<STableViewBase>& InOwnerTableView, TSharedRef<FConcertLogTokenizer> InTokenizer, TSharedRef<FText> InHighlightText)
 {
 	LogEntry = MoveTemp(InLogEntry);
@@ -16,6 +24,9 @@ void SConcertTransportLogRow::Construct(const FArguments& InArgs, TSharedPtr<FCo
 	HighlightText = MoveTemp(InHighlightText);
 	
 	AvatarColor = InArgs._AvatarColor;
+	ScrollToAckLogFunc = InArgs._ScrollToAckLog;
+	ScrollToAckedLogFunc = InArgs._ScrollToAckedLog;
+	check(ScrollToAckLogFunc.IsBound() && ScrollToAckedLogFunc.IsBound());
 	
 	SMultiColumnTableRow<TSharedPtr<FConcertLogEntry>>::Construct(FSuperRowType::FArguments(), InOwnerTableView);
 }
@@ -35,7 +46,9 @@ TSharedRef<SWidget> SConcertTransportLogRow::GenerateWidgetForColumn(const FName
 	}
 	
 	using FColumnWidgetFactoryFunc = TSharedRef<SWidget>(SConcertTransportLogRow::*)(const FName& InColumnName);
-	const TMap<FName, FColumnWidgetFactoryFunc> OverrideFactories = {};
+	const TMap<FName, FColumnWidgetFactoryFunc> OverrideFactories = {
+		{ GET_MEMBER_NAME_CHECKED(FConcertLogMetadata, AckState), &SConcertTransportLogRow::CreateAckColumn }
+	};
 	const FColumnWidgetFactoryFunc FallbackFactoryFunc = &SConcertTransportLogRow::CreateDefaultColumn;
 
 	const FColumnWidgetFactoryFunc* FactoryFunc = OverrideFactories.Find(ColumnName);
@@ -51,13 +64,191 @@ TSharedRef<SWidget> SConcertTransportLogRow::CreateDefaultColumn(const FName& Pr
 		{
 			TSharedPtr<FConcertLogEntry> PinnedEntry = WeakEntry.Pin();
 			TSharedPtr<FConcertLogTokenizer> PinnedTokenizer = WeakTokenizer.Pin();
-			const FProperty* Property = FConcertLog::StaticStruct()->FindPropertyByName(PropertyName);
-			if (ensure(PinnedEntry && PinnedTokenizer && Property))
+			if (!ensure(PinnedEntry && PinnedTokenizer))
+			{
+				return FText::GetEmpty();
+			}
+			
+			const FProperty* ConcertLogProperty = FConcertLog::StaticStruct()->FindPropertyByName(PropertyName);
+			if (ConcertLogProperty)
 			{
-				return FText::FromString(PinnedTokenizer->Tokenize(PinnedEntry->Log, *Property));
+				return FText::FromString(PinnedTokenizer->Tokenize(PinnedEntry->Log, *ConcertLogProperty));
+			}
+
+			const FProperty* ConcertLogEntryProperty = FConcertLogMetadata::StaticStruct()->FindPropertyByName(PropertyName);
+			if (ConcertLogEntryProperty)
+			{
+				return FText::FromString(PinnedTokenizer->Tokenize(PinnedEntry->LogMetaData, *ConcertLogEntryProperty));
 			}
 				
 			return FText::GetEmpty();
 		})
 		.HighlightText_Lambda([this](){ return *HighlightText.Get(); });
 }
+
+TSharedRef<SWidget> SConcertTransportLogRow::CreateAckColumn(const FName& PropertyName)
+{
+	return SNew(SOverlay)
+
+		.ToolTipText_Lambda([this]()
+		{
+			switch (LogEntry->LogMetaData.AckState)
+			{
+				case EConcertLogAckState::Ack: 
+					return LOCTEXT("Ack.Scroll", "This is an ACK. Press to jump to the log this log ACKs.");
+				case EConcertLogAckState::AckReceived: 
+					return LOCTEXT("Acked.Scroll", "ACK received. Press to jump to the log that ACKs this log");
+				case EConcertLogAckState::AckFailure: 
+					return LOCTEXT("AckFailure", "No ACK received: retries exhausted.");
+				case EConcertLogAckState::NotNeeded:
+					return LOCTEXT("Received", "No ACK needed.");
+				case EConcertLogAckState::InProgress:
+					return LOCTEXT("InProgress", "Awaiting ACK...");
+			default:
+				checkNoEntry();
+				return FText::GetEmpty();
+			}
+		})
+
+		+SOverlay::Slot()
+		.HAlign(HAlign_Center)
+		.VAlign(VAlign_Center)
+		[
+			CreateAckInProgressWidget()
+		]
+	
+		+SOverlay::Slot()
+		.HAlign(HAlign_Center)
+		.VAlign(VAlign_Center)
+		[
+			CreateAckFailureWidget()
+		]
+
+		+SOverlay::Slot()
+		.HAlign(HAlign_Center)
+		.VAlign(VAlign_Center)
+		[
+			CreateAckNotNeededWidget()
+		]
+
+		+SOverlay::Slot()
+		.HAlign(HAlign_Center)
+		.VAlign(VAlign_Center)
+		[
+			CreateAckSuccessWidget()
+		]
+
+		+SOverlay::Slot()
+		.HAlign(HAlign_Center)
+		.VAlign(VAlign_Center)
+		[
+			CreateAckWidget()
+		]
+	;
+}
+
+TSharedRef<SWidget> SConcertTransportLogRow::CreateAckInProgressWidget() const
+{
+	return SNew(SScaleBox)
+		.Stretch(EStretch::ScaleToFit)
+		[
+			SNew(SThrobber)
+			.Animate(SThrobber::VerticalAndOpacity)
+			.Visibility_Lambda([this]()
+			{
+				return LogEntry->LogMetaData.AckState == EConcertLogAckState::InProgress
+					? EVisibility::Visible
+					: EVisibility::Hidden;
+			})
+		];
+}
+
+TSharedRef<SWidget> SConcertTransportLogRow::CreateAckSuccessWidget() const
+{
+	return SNew(SButton)
+		.ContentPadding(0.f)
+		.OnPressed_Lambda([this]()
+		{
+			if (LogEntry->LogMetaData.AckingMessageId.IsSet())
+			{
+				ScrollToAckLogFunc.Execute(*LogEntry->LogMetaData.AckingMessageId);
+			}
+		})
+		.Visibility_Lambda([this]()
+		{
+			return LogEntry->LogMetaData.AckState == EConcertLogAckState::AckReceived
+				? EVisibility::Visible
+				: EVisibility::Hidden;
+		})
+		[
+			SNew(SImage)
+			.DesiredSizeOverride(FVector2D(16.f, 16.f))
+			.Image(FConcertServerStyle::Get().GetBrush("Concert.Ack.Success"))
+		];
+}
+
+TSharedRef<SWidget> SConcertTransportLogRow::CreateAckFailureWidget() const
+{
+	return SNew(SImage)
+		.DesiredSizeOverride(FVector2D(16.f, 16.f))
+		.Image(FConcertServerStyle::Get().GetBrush("Concert.Ack.Failure"))
+		.Visibility_Lambda([this]()
+		{
+			return LogEntry->LogMetaData.AckState == EConcertLogAckState::AckFailure
+				? EVisibility::Visible
+				: EVisibility::Hidden;
+		});
+}
+
+TSharedRef<SWidget> SConcertTransportLogRow::CreateAckNotNeededWidget() const
+{
+	return SNew(STextBlock)
+		.Text(LOCTEXT("NotApplicable", "n/a"))
+		.Visibility_Lambda([this]()
+		{
+			return LogEntry->LogMetaData.AckState == EConcertLogAckState::NotNeeded
+				? EVisibility::Visible
+				: EVisibility::Hidden;
+		});
+}
+
+TSharedRef<SWidget> SConcertTransportLogRow::CreateAckWidget() const
+{
+	return SNew(SComboButton)
+		.HasDownArrow(false)
+		.OnGetMenuContent_Lambda([this]()
+		{
+			FMenuBuilder MenuBuilder(true, nullptr);
+			for (const FGuid& MessageId : *LogEntry->LogMetaData.AckedMessageId)
+			{
+				MenuBuilder.AddMenuEntry(
+					FText::FromString(MessageId.ToString(EGuidFormats::DigitsWithHyphens)),
+					FText::GetEmpty(),
+					FSlateIcon(),
+					FUIAction(
+						FExecuteAction::CreateLambda([this, MessageId]()
+						{
+							ScrollToAckedLogFunc.Execute(MessageId);
+						}),
+						FCanExecuteAction::CreateLambda([] { return true; })),
+					NAME_None,
+					EUserInterfaceActionType::Button
+				);
+			}
+			return MenuBuilder.MakeWidget();
+		})
+		.Visibility_Lambda([this]()
+		{
+			return LogEntry->LogMetaData.AckState == EConcertLogAckState::Ack
+				? EVisibility::Visible
+				: EVisibility::Hidden;
+		})
+		.ButtonContent()
+		[
+			SNew(SImage)
+			.DesiredSizeOverride(FVector2D(16.f, 16.f))
+			.Image(FConcertServerStyle::Get().GetBrush("Concert.Ack.Ack"))
+		];
+}
+
+#undef LOCTEXT_NAMESPACE
diff --git a/Engine/Plugins/Developer/Concert/ConcertApp/MultiUserServer/Source/MultiUserServer/Private/Widgets/Clients/Logging/SConcertTransportLogRow.h b/Engine/Plugins/Developer/Concert/ConcertApp/MultiUserServer/Source/MultiUserServer/Private/Widgets/Clients/Logging/SConcertTransportLogRow.h
--- a/Engine/Plugins/Developer/Concert/ConcertApp/MultiUserServer/Source/MultiUserServer/Private/Widgets/Clients/Logging/SConcertTransportLogRow.h
+++ b/Engine/Plugins/Developer/Concert/ConcertApp/MultiUserServer/Source/MultiUserServer/Private/Widgets/Clients/Logging/SConcertTransportLogRow.h
@@ -13,9 +13,13 @@ class SConcertTransportLogRow : public SMultiColumnTableRow<TSharedPtr<FConcertL
 {
 public:
 
+	DECLARE_DELEGATE_OneParam(FScrollToLog, const FGuid& /*MessageId*/);
+
 	SLATE_BEGIN_ARGS(SConcertTransportLogRow)
 	{}
 		SLATE_ARGUMENT(FLinearColor, AvatarColor)
+		SLATE_EVENT(FScrollToLog, ScrollToAckLog)
+		SLATE_EVENT(FScrollToLog, ScrollToAckedLog)
 	SLATE_END_ARGS()
 
 	void Construct(const FArguments& InArgs, TSharedPtr<FConcertLogEntry> InLogEntry, const TSharedRef<STableViewBase>& InOwnerTableView, TSharedRef<FConcertLogTokenizer> InTokenizer, TSharedRef<FText> InHighlightText);
@@ -34,6 +38,16 @@ private:
 	TSharedPtr<FText> HighlightText;
 
 	FLinearColor AvatarColor;
+	FScrollToLog ScrollToAckLogFunc;
+	FScrollToLog ScrollToAckedLogFunc;
 
 	TSharedRef<SWidget> CreateDefaultColumn(const FName& PropertyName);
+	
+	TSharedRef<SWidget> CreateAckColumn(const FName& PropertyName);
+	TSharedRef<SWidget> CreateAckInProgressWidget() const;
+	TSharedRef<SWidget> CreateAckSuccessWidget() const;
+	TSharedRef<SWidget> CreateAckFailureWidget() const;
+	TSharedRef<SWidget> CreateAckNotNeededWidget() const;
+	TSharedRef<SWidget> CreateAckWidget() const;
+	
 };
diff --git a/Engine/Plugins/Developer/Concert/ConcertApp/MultiUserServer/Source/MultiUserServer/Private/Widgets/Clients/Logging/Util/ConcertLogTokenizer.cpp b/Engine/Plugins/Developer/Concert/ConcertApp/MultiUserServer/Source/MultiUserServer/Private/Widgets/Clients/Logging/Util/ConcertLogTokenizer.cpp
--- a/Engine/Plugins/Developer/Concert/ConcertApp/MultiUserServer/Source/MultiUserServer/Private/Widgets/Clients/Logging/Util/ConcertLogTokenizer.cpp
+++ b/Engine/Plugins/Developer/Concert/ConcertApp/MultiUserServer/Source/MultiUserServer/Private/Widgets/Clients/Logging/Util/ConcertLogTokenizer.cpp
@@ -13,6 +13,7 @@ FConcertLogTokenizer::FConcertLogTokenizer(TSharedRef<FEndpointToUserNameCache>
 	: EndpointInfoGetter(EndpointInfoGetter)
 {
 	TokenizerFunctions = {
+			{ FConcertLog::StaticStruct()->FindPropertyByName(GET_MEMBER_NAME_CHECKED(FConcertLog, MessageId)), [this](const FConcertLog& Log) { return TokenizeMessageId(Log); } },
 			{ FConcertLog::StaticStruct()->FindPropertyByName(GET_MEMBER_NAME_CHECKED(FConcertLog, Timestamp)), [this](const FConcertLog& Log) { return TokenizeTimestamp(Log); } },
 			{ FConcertLog::StaticStruct()->FindPropertyByName(GET_MEMBER_NAME_CHECKED(FConcertLog, MessageTypeName)), [this](const FConcertLog& Log) { return TokenizeMessageTypeName(Log); } },
 			{ FConcertLog::StaticStruct()->FindPropertyByName(GET_MEMBER_NAME_CHECKED(FConcertLog, CustomPayloadUncompressedByteSize)), [this](const FConcertLog& Log) { return TokenizeCustomPayloadUncompressedByteSize(Log); } },
@@ -27,7 +28,17 @@ FString FConcertLogTokenizer::Tokenize(const FConcertLog& Data, const FProperty&
 	{
 		return (*CustomTokenizer)(Data);
 	}
-	return TokenizeUsingPropertyExport(Data, ConcertLogProperty);
+	return TokenizeUsingPropertyExport(&Data, ConcertLogProperty);
+}
+
+FString FConcertLogTokenizer::Tokenize(const FConcertLogMetadata& Data, const FProperty& ConcertLogMetadataProperty) const
+{
+	return TokenizeUsingPropertyExport(&Data, ConcertLogMetadataProperty);
+}
+
+FString FConcertLogTokenizer::TokenizeMessageId(const FConcertLog& Data) const
+{
+	return Data.MessageId.ToString(EGuidFormats::DigitsWithHyphens);
 }
 
 FString FConcertLogTokenizer::TokenizeTimestamp(const FConcertLog& Data) const
@@ -57,14 +68,13 @@ FString FConcertLogTokenizer::TokenizeDestinationEndpointId(const FConcertLog& D
 	return GetEndpointDisplayString(Data.DestinationEndpointId);
 }
 
-FString FConcertLogTokenizer::TokenizeUsingPropertyExport(const FConcertLog& Data, const FProperty& ConcertLogProperty) const
+FString FConcertLogTokenizer::TokenizeUsingPropertyExport(const void* ContainerPtr, const FProperty& ConcertLogProperty) const
 {
 	FString Exported;
-	const void* ValuePtr = ConcertLogProperty.ContainerPtrToValuePtr<void>(&Data);
+	const void* ValuePtr = ConcertLogProperty.ContainerPtrToValuePtr<void>(ContainerPtr);
 	const void* DeltaPtr = ValuePtr; // We have no real delta - in this case the API expects the same ptr
 	const bool bSuccess = ConcertLogProperty.ExportText_Direct(Exported, ValuePtr, DeltaPtr, nullptr, PPF_ExternalEditor);
 	check(bSuccess);
-
 	return Exported;
 }
 
diff --git a/Engine/Plugins/Developer/Concert/ConcertApp/MultiUserServer/Source/MultiUserServer/Private/Widgets/Clients/Logging/Util/ConcertLogTokenizer.h b/Engine/Plugins/Developer/Concert/ConcertApp/MultiUserServer/Source/MultiUserServer/Private/Widgets/Clients/Logging/Util/ConcertLogTokenizer.h
--- a/Engine/Plugins/Developer/Concert/ConcertApp/MultiUserServer/Source/MultiUserServer/Private/Widgets/Clients/Logging/Util/ConcertLogTokenizer.h
+++ b/Engine/Plugins/Developer/Concert/ConcertApp/MultiUserServer/Source/MultiUserServer/Private/Widgets/Clients/Logging/Util/ConcertLogTokenizer.h
@@ -6,8 +6,9 @@
 
 class FEndpointToUserNameCache;
 struct FConcertLog;
+struct FConcertLogMetadata;
 
-/** Converts members of FConcertLog into a string. Used e.g. to make search respect the display settings. */
+/** Converts members of FConcertLog and FConcertLogMetadata into a string. Used e.g. to make search respect the display settings. */
 class FConcertLogTokenizer
 {
 public:
@@ -16,13 +17,16 @@ public:
 
 	/** Tokenizes a property of FConcertLog into a string */
 	FString Tokenize(const FConcertLog& Data, const FProperty& ConcertLogProperty) const;
+	/** Tokenizes a property of FConcertLogMetadata into a string */
+	FString Tokenize(const FConcertLogMetadata& Data, const FProperty& ConcertLogMetadataProperty) const;
 
+	// ContainerPtr Tokenizers
+	FString TokenizeMessageId(const FConcertLog& Data) const;
 	FString TokenizeTimestamp(const FConcertLog& Data) const;
 	FString TokenizeMessageTypeName(const FConcertLog& Data) const;
 	FString TokenizeCustomPayloadUncompressedByteSize(const FConcertLog& Data) const;
 	FString TokenizeOriginEndpointId(const FConcertLog& Data) const;
 	FString TokenizeDestinationEndpointId(const FConcertLog& Data) const;
-	FString TokenizeUsingPropertyExport(const FConcertLog& Data, const FProperty& ConcertLogProperty) const;
 
 private:
 
@@ -35,4 +39,6 @@ private:
 	TSharedRef<FEndpointToUserNameCache> EndpointInfoGetter;
 
 	FString GetEndpointDisplayString(const FGuid& EndpointId) const;
+	
+	FString TokenizeUsingPropertyExport(const void* ContainerPtr, const FProperty& ConcertLogProperty) const;
 };
diff --git a/Engine/Plugins/Developer/Concert/ConcertApp/MultiUserServer/Source/MultiUserServer/Private/Widgets/Clients/Logging/Util/LogAckTracker.cpp b/Engine/Plugins/Developer/Concert/ConcertApp/MultiUserServer/Source/MultiUserServer/Private/Widgets/Clients/Logging/Util/LogAckTracker.cpp
new file mode 100644
--- /dev/null
+++ b/Engine/Plugins/Developer/Concert/ConcertApp/MultiUserServer/Source/MultiUserServer/Private/Widgets/Clients/Logging/Util/LogAckTracker.cpp
@@ -0,0 +1,84 @@
+// Copyright Epic Games, Inc. All Rights Reserved.
+
+#include "LogAckTracker.h"
+
+#include "ConcertTransportEvents.h"
+#include "IConcertServer.h"
+
+#include "UObject/StructOnScope.h"
+#include "Widgets/Clients/Logging/Source/IConcertLogSource.h"
+
+FLogAckTracker::FLogAckTracker(TSharedRef<IConcertLogSource> LogSource, TSharedRef<IConcertServer> Server)
+	: LogSource(MoveTemp(LogSource))
+	, Server(MoveTemp(Server))
+{
+	LogSource->OnLogEntryAdded().AddRaw(this, &FLogAckTracker::OnLogEntryProduced);
+	Server->OnConcertMessageAcknowledgementReceived().AddRaw(this, &FLogAckTracker::OnAckProcessed);
+}
+
+FLogAckTracker::~FLogAckTracker()
+{
+	LogSource->OnLogEntryAdded().RemoveAll(this);
+	Server->OnConcertMessageAcknowledgementReceived().RemoveAll(this);
+}
+
+void FLogAckTracker::OnLogEntryProduced(const TSharedRef<FConcertLogEntry>& LogEntry)
+{
+	const bool bIsAck = LogEntry->Log.MessageTypeName == FConcertAckData::StaticStruct()->GetFName();
+	if (bIsAck)
+	{
+		if (LogEntry->Log.MessageAction == EConcertLogMessageAction::Process)
+		{
+			LogEntry->LogMetaData.AckState = EConcertLogAckState::Ack;
+			TSet<FGuid> AckedMessageIds;
+			AckingMessageToAckedMessages.RemoveAndCopyValue(LogEntry->Log.MessageId, AckedMessageIds);
+			LogEntry->LogMetaData.AckedMessageId = AckedMessageIds;
+		}
+		return;
+	}
+	
+	if (LogEntry->Log.MessageAction == EConcertLogMessageAction::TimeOut)
+	{
+		if (const FPendingAckInfo* TimedOutPending = PendingAckLogs.Find(LogEntry->Log.MessageId); TimedOutPending )
+		{
+			for (const TWeakPtr<FConcertLogEntry>& TimedOut : TimedOutPending->LogEntries)
+			{
+				TimedOut.Pin()->LogMetaData.AckState = EConcertLogAckState::AckFailure;
+			}
+		}
+		PendingAckLogs.Remove(LogEntry->Log.MessageId);
+		return;
+	}
+
+	const TSet<EConcertLogMessageAction> AckNeeded = { EConcertLogMessageAction::Send };
+	const bool bIsReliable = LogEntry->Log.ChannelId != FConcertMessageData::UnreliableChannelId;
+	if (bIsReliable && AckNeeded.Contains(LogEntry->Log.MessageAction))
+	{
+		LogEntry->LogMetaData.AckState = EConcertLogAckState::InProgress;
+		PendingAckLogs.FindOrAdd(LogEntry->Log.MessageId).LogEntries.Add(LogEntry);
+	}
+	else
+	{
+		LogEntry->LogMetaData.AckState = EConcertLogAckState::NotNeeded;
+	}
+}
+
+void FLogAckTracker::OnAckProcessed(
+	const FConcertEndpointContext& LocalEndpointId,
+	const FConcertEndpointContext& RemoteEndpointId,
+	const TSharedRef<IConcertMessage>& AckedMessage,
+	const FConcertMessageContext& MessageContext)
+{
+	if (const FPendingAckInfo* Pending = PendingAckLogs.Find(AckedMessage->GetMessageId()); Pending)
+	{
+		AckingMessageToAckedMessages.FindOrAdd(MessageContext.Message->MessageId).Add(AckedMessage->GetMessageId());
+		for (const TWeakPtr<FConcertLogEntry>& LogEntry : Pending->LogEntries)
+		{
+			TSharedPtr<FConcertLogEntry> Pinned = LogEntry.Pin();
+			Pinned->LogMetaData.AckState = EConcertLogAckState::AckReceived;
+			Pinned->LogMetaData.AckingMessageId = MessageContext.Message->MessageId;
+		}
+		
+		PendingAckLogs.Remove(AckedMessage->GetMessageId());
+	}
+}
diff --git a/Engine/Plugins/Developer/Concert/ConcertApp/MultiUserServer/Source/MultiUserServer/Private/Widgets/Clients/Logging/Util/LogAckTracker.h b/Engine/Plugins/Developer/Concert/ConcertApp/MultiUserServer/Source/MultiUserServer/Private/Widgets/Clients/Logging/Util/LogAckTracker.h
new file mode 100644
--- /dev/null
+++ b/Engine/Plugins/Developer/Concert/ConcertApp/MultiUserServer/Source/MultiUserServer/Private/Widgets/Clients/Logging/Util/LogAckTracker.h
@@ -0,0 +1,43 @@
+// Copyright Epic Games, Inc. All Rights Reserved.
+
+#pragma once
+
+#include "CoreMinimal.h"
+#include "ConcertMessageContext.h"
+#include "Widgets/Clients/Logging/ConcertLogEntry.h"
+
+class FStructOnScope;
+class IConcertLogSource;
+class IConcertServer;
+
+struct FConcertLogEntry;
+
+/** Tracks ack logs and updates the  */
+class FLogAckTracker
+{
+public:
+	FLogAckTracker(TSharedRef<IConcertLogSource> LogSource, TSharedRef<IConcertServer> Server);
+	~FLogAckTracker();
+
+private:
+
+	/** We analyse & track newly added logs */
+	const TSharedRef<IConcertLogSource> LogSource;
+	/** We listen to received ACKs from the server. */
+	TSharedRef<IConcertServer> Server;
+
+	struct FPendingAckInfo
+	{
+		/** There can be several logs with the same Message ID */
+		TArray<TWeakPtr<FConcertLogEntry>> LogEntries;
+	};
+	
+	/** Logs we're waiting on an ACK for */
+	TMap<FGuid, FPendingAckInfo> PendingAckLogs;
+
+	/** Maps an ACK log to the message IDs it ACKs. */
+	TMap<FGuid, TSet<FGuid>> AckingMessageToAckedMessages;
+	
+	void OnLogEntryProduced(const TSharedRef<FConcertLogEntry>& LogEntry);
+	void OnAckProcessed(const FConcertEndpointContext& LocalEndpointId, const FConcertEndpointContext& RemoteEndpointId, const TSharedRef<IConcertMessage>& AckedMessage, const FConcertMessageContext& MessageContext);
+};
