Show clients without a session in session browser

diff --git a/Engine/Plugins/Developer/Concert/ConcertApp/MultiUserServer/Source/MultiUserServer/Private/Widgets/Clients/Browser/ClientBrowserItem.h b/Engine/Plugins/Developer/Concert/ConcertApp/MultiUserServer/Source/MultiUserServer/Private/Widgets/Clients/Browser/ClientBrowserItem.h
--- a/Engine/Plugins/Developer/Concert/ConcertApp/MultiUserServer/Source/MultiUserServer/Private/Widgets/Clients/Browser/ClientBrowserItem.h
+++ b/Engine/Plugins/Developer/Concert/ConcertApp/MultiUserServer/Source/MultiUserServer/Private/Widgets/Clients/Browser/ClientBrowserItem.h
@@ -8,14 +8,29 @@
 
 namespace UE::MultiUserServer
 {
+	DECLARE_DELEGATE_RetVal(TOptional<FConcertClientInfo>, FGetClientInfo);
+	
 	struct FClientBrowserItem
 	{
-		FConcertSessionClientInfo ClientInfo;
+		/** Getter for getting the latest client info. */
+		FGetClientInfo GetClientInfo;
+		
+		/** Address of this client as used by the UDP messaging sytem */
 		FMessageAddress ClientAddress;
+		/** ID of this client in the UDP messaging system */
+		FGuid MessageNodeId;
 
-		FClientBrowserItem(const FConcertSessionClientInfo& ClientInfo, FMessageAddress ClientAddress)
-			: ClientInfo(ClientInfo)
+		/** The session this client currently is in */
+		TOptional<FGuid> CurrentSession;
+
+		/** Whether this client is no longer connected with the server */
+		bool bIsDisconnected = false;
+
+		FClientBrowserItem(FGetClientInfo GetClientInfo, const FMessageAddress& ClientAddress, const FGuid& MessageNodeId, TOptional<FGuid> CurrentSession = {})
+			: GetClientInfo(MoveTemp(GetClientInfo))
 			, ClientAddress(ClientAddress)
+			, MessageNodeId(MessageNodeId)
+			, CurrentSession(CurrentSession)
 		{}
 	};
 }
diff --git a/Engine/Plugins/Developer/Concert/ConcertApp/MultiUserServer/Source/MultiUserServer/Private/Widgets/Clients/Browser/Models/ClientBrowserModel.cpp b/Engine/Plugins/Developer/Concert/ConcertApp/MultiUserServer/Source/MultiUserServer/Private/Widgets/Clients/Browser/Models/ClientBrowserModel.cpp
--- a/Engine/Plugins/Developer/Concert/ConcertApp/MultiUserServer/Source/MultiUserServer/Private/Widgets/Clients/Browser/Models/ClientBrowserModel.cpp
+++ b/Engine/Plugins/Developer/Concert/ConcertApp/MultiUserServer/Source/MultiUserServer/Private/Widgets/Clients/Browser/Models/ClientBrowserModel.cpp
@@ -3,12 +3,18 @@
 #include "ClientBrowserModel.h"
 
 #include "ConcertServerEvents.h"
-#include "ConcertUtil.h"
 #include "IConcertServer.h"
+#include "Widgets/Clients/Browser/ClientBrowserItem.h"
+#include "Widgets/Clients/Util/EndpointToUserNameCache.h"
+
+#include "Algo/AnyOf.h"
 #include "Algo/Transform.h"
+#include "Features/IModularFeatures.h"
+#include "INetworkMessagingExtension.h"
 
-UE::MultiUserServer::FClientBrowserModel::FClientBrowserModel(TSharedRef<IConcertServer> InServer)
+UE::MultiUserServer::FClientBrowserModel::FClientBrowserModel(TSharedRef<IConcertServer> InServer, TSharedRef<FEndpointToUserNameCache> ClientInfoCache)
 	: Server(MoveTemp(InServer))
+	, ClientInfoCache(MoveTemp(ClientInfoCache))
 {
 	ConcertServerEvents::OnLiveSessionCreated().AddRaw(this, &FClientBrowserModel::OnLiveSessionCreated);
 	ConcertServerEvents::OnLiveSessionDestroyed().AddRaw(this, &FClientBrowserModel::OnLiveSessionDestroyed);
@@ -16,6 +22,12 @@ UE::MultiUserServer::FClientBrowserModel::FClientBrowserModel(TSharedRef<IConcer
 	{
 		SubscribeToClientConnectionEvents(LiveSession.ToSharedRef());
 	}
+
+	for (const FConcertEndpointContext& Context : Server->GetRemoteAdminEndpoints())
+	{
+		AddClientAdminEndpoint(Context);
+	}
+	Server->OnRemoteEndpointConnectionChanged().AddRaw(this, &FClientBrowserModel::OnAdminClientEndpointConnectionChanged);
 }
 
 UE::MultiUserServer::FClientBrowserModel::~FClientBrowserModel()
@@ -26,6 +38,8 @@ UE::MultiUserServer::FClientBrowserModel::~FClientBrowserModel()
 	{
 		UnsubscribeFromClientConnectionEvents(LiveSession.ToSharedRef());
 	}
+	
+	Server->OnRemoteEndpointConnectionChanged().RemoveAll(this);
 }
 
 TSet<FGuid> UE::MultiUserServer::FClientBrowserModel::GetSessions() const
@@ -44,18 +58,30 @@ TOptional<FConcertSessionInfo> UE::MultiUserServer::FClientBrowserModel::GetSess
 		: TOptional<FConcertSessionInfo>();
 }
 
-TArray<FConcertSessionClientInfo> UE::MultiUserServer::FClientBrowserModel::GetSessionClients(const FGuid& SessionID) const
+void UE::MultiUserServer::FClientBrowserModel::SetKeepClientsAfterDisconnect(bool bNewValue)
 {
-	return ConcertUtil::GetSessionClients(*Server, SessionID);
-}
+	if (bNewValue == bKeepClientsAfterDisconnect)
+	{
+		return;
+	}
 
-FMessageAddress UE::MultiUserServer::FClientBrowserModel::GetClientAddress(const FGuid& ClientEndpointId) const
-{
-	const TSharedPtr<IConcertServerSession> ServerSession = ConcertUtil::GetLiveSessionClientConnectedTo(*Server, ClientEndpointId);
-	return ServerSession->GetClientAddress(ClientEndpointId);
+	bKeepClientsAfterDisconnect = bNewValue;
+
+	if (bKeepClientsAfterDisconnect)
+	{
+		return;
+	}
+	for (auto It = Clients.CreateIterator(); It; ++It)
+	{
+		if (It->Get()->bIsDisconnected)
+		{
+			OnClientListChanged().Broadcast(*It, EClientUpdateType::Removed);
+			It.RemoveCurrent();
+		}
+	}
 }
 
-void UE::MultiUserServer::FClientBrowserModel::OnLiveSessionCreated(bool bSuccess, const IConcertServer& InServer, TSharedRef<IConcertServerSession> InLiveSession) const
+void UE::MultiUserServer::FClientBrowserModel::OnLiveSessionCreated(bool bSuccess, const IConcertServer& InServer, TSharedRef<IConcertServerSession> InLiveSession)
 {
 	if (bSuccess)
 	{
@@ -70,17 +96,146 @@ void UE::MultiUserServer::FClientBrowserModel::OnLiveSessionDestroyed(const ICon
 	OnSessionCreatedEvent.Broadcast(InLiveSession->GetId());
 }
 
-void UE::MultiUserServer::FClientBrowserModel::SubscribeToClientConnectionEvents(const TSharedRef<IConcertServerSession>& InLiveSession) const
+void UE::MultiUserServer::FClientBrowserModel::SubscribeToClientConnectionEvents(const TSharedRef<IConcertServerSession>& InLiveSession)
 {
 	InLiveSession->OnSessionClientChanged().AddRaw(this, &FClientBrowserModel::OnClientListUpdated);
 }
 
-void UE::MultiUserServer::FClientBrowserModel::UnsubscribeFromClientConnectionEvents( const TSharedRef<IConcertServerSession>& InLiveSession) const
+void UE::MultiUserServer::FClientBrowserModel::UnsubscribeFromClientConnectionEvents(const TSharedRef<IConcertServerSession>& InLiveSession) const
 {
 	InLiveSession->OnSessionClientChanged().RemoveAll(this);
 }
 
-void UE::MultiUserServer::FClientBrowserModel::OnClientListUpdated(IConcertServerSession& Session, EConcertClientStatus Status, const FConcertSessionClientInfo& ClientInfo) const
+void UE::MultiUserServer::FClientBrowserModel::OnClientListUpdated(IConcertServerSession& Session, EConcertClientStatus Status, const FConcertSessionClientInfo& ClientInfo)
+{
+	switch (Status)
+	{
+		case EConcertClientStatus::Connected:
+			OnClientJoinSession(Session, ClientInfo.ClientEndpointId);
+			break;
+		case EConcertClientStatus::Disconnected:
+			OnClientLeaveSession(Session, ClientInfo.ClientEndpointId);
+			break;
+		case EConcertClientStatus::Updated:
+			break;
+		default:
+			checkNoEntry();
+	}
+}
+
+void UE::MultiUserServer::FClientBrowserModel::OnClientJoinSession(IConcertServerSession& Session, const FGuid& EndpointId)
+{
+	UpdateClientSessionId(Session, EndpointId, Session.GetId());
+}
+
+void UE::MultiUserServer::FClientBrowserModel::OnClientLeaveSession(IConcertServerSession& Session, const FGuid& EndpointId)
+{
+	UpdateClientSessionId(Session, EndpointId, {});
+}
+
+void UE::MultiUserServer::FClientBrowserModel::UpdateClientSessionId(IConcertServerSession& Session, const FGuid& EndpointId, TOptional<FGuid> SessionId)
+{
+	IModularFeatures& ModularFeatures = IModularFeatures::Get();
+	if (!ensure(ModularFeatures.IsModularFeatureAvailable(INetworkMessagingExtension::ModularFeatureName)))
+	{
+		return;
+	}
+	const INetworkMessagingExtension& Extension = ModularFeatures.GetModularFeature<INetworkMessagingExtension>(INetworkMessagingExtension::ModularFeatureName);
+	
+	const FMessageAddress Address = Session.GetClientAddress(EndpointId);
+	// Invalid if client disconnects while in a session
+	if (!Address.IsValid())
+	{
+		return;
+	}
+	
+	const FGuid NodeId = Extension.GetNodeIdFromAddress(Address);
+	const int32 Index = Clients.IndexOfByPredicate([&NodeId](const TSharedPtr<FClientBrowserItem>& Item)
+	{
+		return Item->MessageNodeId == NodeId;
+	});
+	if (Clients.IsValidIndex(Index))
+	{
+		Clients[Index]->CurrentSession = SessionId;
+	}
+}
+
+void UE::MultiUserServer::FClientBrowserModel::OnAdminClientEndpointConnectionChanged(const FConcertEndpointContext& Context, EConcertRemoteEndpointConnection ConnectionState)
 {
-	OnClientListChangedEvent.Broadcast(Session.GetId(), Status, ClientInfo);
+	switch (ConnectionState)
+	{
+	case EConcertRemoteEndpointConnection::Discovered:
+		AddClientAdminEndpoint(Context);
+		break;
+	case EConcertRemoteEndpointConnection::TimedOut:
+	case EConcertRemoteEndpointConnection::ClosedRemotely:
+		RemoveClientAdminEndpoint(Context);
+		break;
+	default:
+		checkNoEntry();
+	}
 }
+
+void UE::MultiUserServer::FClientBrowserModel::AddClientAdminEndpoint(const FConcertEndpointContext& Context)
+{
+	const FMessageAddress Address = Server->GetRemoteAddress(Context.EndpointId);
+
+	if (!ensure(Address.IsValid()))
+	{
+		return;
+	}
+	
+	IModularFeatures& ModularFeatures = IModularFeatures::Get();
+	const INetworkMessagingExtension* MessagingExtension = ModularFeatures.IsModularFeatureAvailable(INetworkMessagingExtension::ModularFeatureName)
+		? &ModularFeatures.GetModularFeature<INetworkMessagingExtension>(INetworkMessagingExtension::ModularFeatureName)
+		: nullptr;
+	if (!ensure(MessagingExtension))
+	{
+		return;
+	}
+	const FGuid NodeId = MessagingExtension->GetNodeIdFromAddress(Address);
+	const TSharedPtr<FClientBrowserItem>* ExistingItem = Clients.FindByPredicate([&NodeId](const TSharedPtr<FClientBrowserItem>& Item)
+	{
+		return Item->MessageNodeId == NodeId;
+	});
+	if (ExistingItem)
+	{
+		ExistingItem->Get()->bIsDisconnected = false;
+		return;
+	}
+	
+	const TSharedRef<FClientBrowserItem> Item = MakeShared<FClientBrowserItem>(
+		FGetClientInfo::CreateLambda([this, NodeId]()
+		{
+			return ClientInfoCache->GetClientInfoFromNodeId(NodeId);
+		}),
+		Address,
+		NodeId
+		);
+	Clients.Add(Item);
+	EndpointToNodeId.Add(Context.EndpointId, NodeId);
+	OnClientListChangedEvent.Broadcast(Item, EClientUpdateType::Added);
+}
+
+void UE::MultiUserServer::FClientBrowserModel::RemoveClientAdminEndpoint(const FConcertEndpointContext& Context)
+{
+	const FMessageNodeId NodeId = EndpointToNodeId[Context.EndpointId];
+	const int32 Index = Clients.IndexOfByPredicate([&NodeId](const TSharedPtr<FClientBrowserItem>& Item)
+	{
+		return Item->MessageNodeId == NodeId;
+	});
+	
+	if (ensure(Clients.IsValidIndex(Index)))
+	{
+		EndpointToNodeId.Remove(Context.EndpointId);
+		const TSharedPtr<FClientBrowserItem> Item = Clients[Index]; 
+		if (bKeepClientsAfterDisconnect)
+		{
+			Item->bIsDisconnected = true;
+			return;
+		}
+		
+		Clients.RemoveAt(Index);
+		OnClientListChangedEvent.Broadcast(Item, EClientUpdateType::Removed);
+	}
+}
\ No newline at end of file
diff --git a/Engine/Plugins/Developer/Concert/ConcertApp/MultiUserServer/Source/MultiUserServer/Private/Widgets/Clients/Browser/Models/ClientBrowserModel.h b/Engine/Plugins/Developer/Concert/ConcertApp/MultiUserServer/Source/MultiUserServer/Private/Widgets/Clients/Browser/Models/ClientBrowserModel.h
--- a/Engine/Plugins/Developer/Concert/ConcertApp/MultiUserServer/Source/MultiUserServer/Private/Widgets/Clients/Browser/Models/ClientBrowserModel.h
+++ b/Engine/Plugins/Developer/Concert/ConcertApp/MultiUserServer/Source/MultiUserServer/Private/Widgets/Clients/Browser/Models/ClientBrowserModel.h
@@ -5,6 +5,9 @@
 #include "CoreMinimal.h"
 #include "IClientBrowserModel.h"
 
+enum class EConcertRemoteEndpointConnection : uint8;
+struct FConcertEndpointContext;
+class FEndpointToUserNameCache;
 class IConcertServerSession;
 class IConcertServer;
 
@@ -14,14 +17,15 @@ namespace UE::MultiUserServer
 	{
 	public:
 
-		FClientBrowserModel(TSharedRef<IConcertServer> InServer);
+		FClientBrowserModel(TSharedRef<IConcertServer> InServer, TSharedRef<FEndpointToUserNameCache> ClientInfoCache);
 		virtual ~FClientBrowserModel() override;
 
 		//~ Begin IClientBrowserModel Interface
 		virtual TSet<FGuid> GetSessions() const override;
 		virtual TOptional<FConcertSessionInfo> GetSessionInfo(const FGuid& SessionID) const override;
-		virtual TArray<FConcertSessionClientInfo> GetSessionClients(const FGuid& SessionID) const override;
-		virtual FMessageAddress GetClientAddress(const FGuid& ClientEndpointId) const override;
+		virtual const TArray<TSharedPtr<FClientBrowserItem>>& GetItems() const override { return Clients; }
+		virtual void SetKeepClientsAfterDisconnect(bool bNewValue) override;
+		virtual bool ShouldKeepClientsAfterDisconnect() const override { return bKeepClientsAfterDisconnect; }
 		virtual FOnClientListChanged& OnClientListChanged() override { return OnClientListChangedEvent; }
 		virtual FOnSessionListChanged& OnSessionCreated() override { return OnSessionCreatedEvent; }
 		virtual FOnSessionListChanged& OnSessionDestroyed() override { return OnSessionDestroyedEvent; }
@@ -29,18 +33,37 @@ namespace UE::MultiUserServer
 
 	private:
 
+		using FClientEndpointId = FGuid;
+		using FMessageNodeId = FGuid;
+
 		TSharedRef<IConcertServer> Server;
+		TSharedRef<FEndpointToUserNameCache> ClientInfoCache;
+
+		TArray<TSharedPtr<FClientBrowserItem>> Clients;
+		/** Needed for when an admin endpoint disconnects - querying won't work anymore then. */
+		TMap<FClientEndpointId, FMessageNodeId> EndpointToNodeId;
+
+		bool bKeepClientsAfterDisconnect = true;
 
 		FOnClientListChanged OnClientListChangedEvent;
 		FOnSessionListChanged OnSessionCreatedEvent;
 		FOnSessionListChanged OnSessionDestroyedEvent;
 
-		void OnLiveSessionCreated(bool bSuccess, const IConcertServer& InServer, TSharedRef<IConcertServerSession> InLiveSession) const;
+		// Session events
+		void OnLiveSessionCreated(bool bSuccess, const IConcertServer& InServer, TSharedRef<IConcertServerSession> InLiveSession);
 		void OnLiveSessionDestroyed(const IConcertServer& InServer, TSharedRef<IConcertServerSession> InLiveSession) const;
-		
-		void SubscribeToClientConnectionEvents(const TSharedRef<IConcertServerSession>& InLiveSession) const;
+		void SubscribeToClientConnectionEvents(const TSharedRef<IConcertServerSession>& InLiveSession);
 		void UnsubscribeFromClientConnectionEvents(const TSharedRef<IConcertServerSession>& InLiveSession) const;
-		void OnClientListUpdated(IConcertServerSession& Session, EConcertClientStatus Status, const FConcertSessionClientInfo& ClientInfo) const;
+
+		// Handle session clients
+		void OnClientListUpdated(IConcertServerSession& Session, EConcertClientStatus Status, const FConcertSessionClientInfo& ClientInfo);
+		void OnClientJoinSession(IConcertServerSession& Session, const FGuid& EndpointId);
+		void OnClientLeaveSession(IConcertServerSession& Session, const FGuid& EndpointId);
+		void UpdateClientSessionId(IConcertServerSession& Session, const FGuid& EndpointId, TOptional<FGuid> SessionId);
+
+		void OnAdminClientEndpointConnectionChanged(const FConcertEndpointContext& Context, EConcertRemoteEndpointConnection ConnectionState);
+		void AddClientAdminEndpoint(const FConcertEndpointContext& Context);
+		void RemoveClientAdminEndpoint(const FConcertEndpointContext& Context);
 	};
 }
 
diff --git a/Engine/Plugins/Developer/Concert/ConcertApp/MultiUserServer/Source/MultiUserServer/Private/Widgets/Clients/Browser/Models/IClientBrowserModel.h b/Engine/Plugins/Developer/Concert/ConcertApp/MultiUserServer/Source/MultiUserServer/Private/Widgets/Clients/Browser/Models/IClientBrowserModel.h
--- a/Engine/Plugins/Developer/Concert/ConcertApp/MultiUserServer/Source/MultiUserServer/Private/Widgets/Clients/Browser/Models/IClientBrowserModel.h
+++ b/Engine/Plugins/Developer/Concert/ConcertApp/MultiUserServer/Source/MultiUserServer/Private/Widgets/Clients/Browser/Models/IClientBrowserModel.h
@@ -11,24 +11,35 @@ struct FMessageAddress;
 
 namespace UE::MultiUserServer
 {
+	struct FClientBrowserItem;
+	
 	/** Decouples the UI from the server functions. */
 	class IClientBrowserModel
 	{
 	public:
 
+		enum class EClientUpdateType : uint8
+		{
+			Added,
+			Removed
+		};
+		
+		DECLARE_MULTICAST_DELEGATE_TwoParams(FOnClientListChanged, TSharedPtr<FClientBrowserItem>, EClientUpdateType);
+		DECLARE_MULTICAST_DELEGATE_OneParam(FOnSessionListChanged, const FGuid& /*SessionId*/);
+
 		/** Gets the IDs of all availabel IDs */
 		virtual TSet<FGuid> GetSessions() const = 0;
 		/** Gets more info about a session returned by GetSessions */
 		virtual TOptional<FConcertSessionInfo> GetSessionInfo(const FGuid& SessionID) const = 0;
-		/** Gets the clients connected to a session returned by GetSessions */
-		virtual TArray<FConcertSessionClientInfo> GetSessionClients(const FGuid& SessionID) const = 0;
-		/** Gets the network address of a given client */
-		virtual FMessageAddress GetClientAddress(const FGuid& ClientEndpointId) const = 0; 
+
+		virtual const TArray<TSharedPtr<FClientBrowserItem>>& GetItems() const = 0;
+
+		virtual void SetKeepClientsAfterDisconnect(bool bNewValue) = 0;
+		virtual bool ShouldKeepClientsAfterDisconnect() const = 0;
 		
-		DECLARE_MULTICAST_DELEGATE_ThreeParams(FOnClientListChanged, const FGuid& /*SessionId*/, EConcertClientStatus /*UpdateType*/, const FConcertSessionClientInfo& /*ClientInfo*/);
+		/** Called when GeItems() changes its content */
 		virtual FOnClientListChanged& OnClientListChanged() = 0;
 
-		DECLARE_MULTICAST_DELEGATE_OneParam(FOnSessionListChanged, const FGuid& /*SessionId*/);
 		virtual FOnSessionListChanged& OnSessionCreated() = 0;
 		virtual FOnSessionListChanged& OnSessionDestroyed() = 0;
 		
diff --git a/Engine/Plugins/Developer/Concert/ConcertApp/MultiUserServer/Source/MultiUserServer/Private/Widgets/Clients/Browser/SConcertClientBrowser.cpp b/Engine/Plugins/Developer/Concert/ConcertApp/MultiUserServer/Source/MultiUserServer/Private/Widgets/Clients/Browser/SConcertClientBrowser.cpp
--- a/Engine/Plugins/Developer/Concert/ConcertApp/MultiUserServer/Source/MultiUserServer/Private/Widgets/Clients/Browser/SConcertClientBrowser.cpp
+++ b/Engine/Plugins/Developer/Concert/ConcertApp/MultiUserServer/Source/MultiUserServer/Private/Widgets/Clients/Browser/SConcertClientBrowser.cpp
@@ -4,16 +4,22 @@
 
 #include "ClientBrowserItem.h"
 #include "ConcertServerStyle.h"
-#include "SConcertClientBrowserItem.h"
+#include "MultiUserServerModule.h"
 #include "Models/IClientBrowserModel.h"
+#include "Models/IClientNetworkStatisticsModel.h"
+#include "SConcertClientBrowserItem.h"
+#include "Algo/Accumulate.h"
 
 #include "Algo/AnyOf.h"
+#include "Algo/Count.h"
+#include "Dialog/SMessageDialog.h"
 #include "Framework/MultiBox/MultiBoxBuilder.h"
-#include "Models/IClientNetworkStatisticsModel.h"
-#include "Widgets/SBoxPanel.h"
-#include "Widgets/Input/SComboBox.h"
+#include "Widgets/Input/SCheckBox.h"
 #include "Widgets/Input/SSearchBox.h"
+#include "Widgets/SBoxPanel.h"
+#include "Widgets/Text/STextBlock.h"
 #include "Widgets/Views/STileView.h"
+#include "Window/ModalWindowManager.h"
 
 #define LOCTEXT_NAMESPACE "UnrealMultiUserUI.SConcertClientBrowser"
 
@@ -115,9 +121,72 @@ TSharedRef<SWidget> UE::MultiUserServer::SConcertClientBrowser::CreateSearchArea
 		.AutoWidth()
 		[
 			InArgs._RightOfSearch.Widget
+		]
+
+		+SHorizontalBox::Slot()
+		.AutoWidth()
+		.HAlign(HAlign_Right)
+		[
+			CreateKeepDisconnectedClients()
 		];
 }
 
+TSharedRef<SWidget> UE::MultiUserServer::SConcertClientBrowser::CreateKeepDisconnectedClients()
+{
+	return SNew(SHorizontalBox)
+		.ToolTipText(LOCTEXT("KeepDisconnectedClients.Tooltip", "Whether to keep clients that have disconnected in memory. This may be useful in unstable networks when you want to analyse why clients keep disconnecting."))
+		
+		+SHorizontalBox::Slot()
+		.AutoWidth()
+		.VAlign(VAlign_Center)
+		[
+			SNew(STextBlock)
+			.Text(LOCTEXT("KeepDisconnectedClients.Label", "Keep Disconnected"))
+			]
+
+		+SHorizontalBox::Slot()
+		.AutoWidth()
+		.Padding(4.f, 0.f, 0.f, 0.f)
+		[
+			SNew(SCheckBox)
+			.IsChecked_Lambda([this](){ return BrowserModel->ShouldKeepClientsAfterDisconnect() ? ECheckBoxState::Checked : ECheckBoxState::Unchecked; })
+			.OnCheckStateChanged_Lambda([this](ECheckBoxState NewState)
+			{
+				const bool bKeepDisconnected = NewState == ECheckBoxState::Checked;
+				if (bKeepDisconnected)
+				{
+					BrowserModel->SetKeepClientsAfterDisconnect(true);
+					return;
+				}
+				
+				const int32 NumDisconnected = Algo::CountIf(BrowserModel->GetItems(), [](const TSharedPtr<FClientBrowserItem>& Item) { return Item->bIsDisconnected; });
+				if (NumDisconnected == 0)
+				{
+					BrowserModel->SetKeepClientsAfterDisconnect(false);
+					return;
+				}
+				
+				const TSharedRef<SMessageDialog> Dialog = SNew(SMessageDialog)
+					.Title(LOCTEXT("RemoveDisconnectedClients.Title", "Remove disconnected clients?"))
+					.IconBrush("Icons.WarningWithColor.Large")
+					.Message(FText::Format(LOCTEXT("RemoveDisconnectedClients.MessageFmt", "There are {0} disconnected clients. If you proceed, these clients will be removed from the session browser; opened log tabs will remain open.\nProceed?"), NumDisconnected))
+					.UseScrollBox(false)
+					.Buttons({
+						SMessageDialog::FButton(LOCTEXT("RemoveButton", "Remove"))
+							.SetOnClicked(FSimpleDelegate::CreateLambda([this]()
+							{
+								BrowserModel->SetKeepClientsAfterDisconnect(false);
+							})),
+						SMessageDialog::FButton(LOCTEXT("CancelButton", "Keep"))
+							.SetPrimary(true)
+							.SetFocus()
+					});
+				FConcertServerUIModule::Get().GetModalWindowManager()->ShowFakeModalWindow(Dialog);
+			})
+		]
+	;
+}
+
 TSharedRef<SWidget> UE::MultiUserServer::SConcertClientBrowser::CreateTileView()
 {
 	constexpr int32 Height = 200;
@@ -151,44 +220,24 @@ void UE::MultiUserServer::SConcertClientBrowser::OnSessionDestroyed(const FGuid&
 	}
 }
 
-void UE::MultiUserServer::SConcertClientBrowser::OnClientListChanged(const FGuid& SessionId, EConcertClientStatus UpdateType, const FConcertSessionClientInfo& ClientInfo)
+void UE::MultiUserServer::SConcertClientBrowser::OnClientListChanged(TSharedPtr<FClientBrowserItem> Item, IClientBrowserModel::EClientUpdateType UpdateType)
 {
 	switch (UpdateType)
 	{
-	case EConcertClientStatus::Connected:
-		DisplayedClients.Add(MakeShared<FClientBrowserItem>(ClientInfo, BrowserModel->GetClientAddress(ClientInfo.ClientEndpointId)));
-		TileView->RequestListRefresh();
-		break;
-	case EConcertClientStatus::Disconnected:
-		RemoveClient(ClientInfo);
+	case IClientBrowserModel::EClientUpdateType::Added:
+		if (PassesFilter(Item))
+		{
+			DisplayedClients.Add(Item);
+		}
 		break;
-	case EConcertClientStatus::Updated:
-		UpdateClientInfo(ClientInfo);
+	case IClientBrowserModel::EClientUpdateType::Removed:
+		DisplayedClients.Remove(Item);
 		break;
-	default: ;
+	default:
+		checkNoEntry();
 	}
-}
 
-void UE::MultiUserServer::SConcertClientBrowser::RemoveClient(const FConcertSessionClientInfo& ClientInfo)
-{
-	const int32 Index = DisplayedClients.IndexOfByPredicate([&ClientInfo](const TSharedPtr<FClientBrowserItem>& Item)
-		{
-			return Item->ClientInfo.ClientEndpointId == ClientInfo.ClientEndpointId;
-		});
-	if (ensure(Index != INDEX_NONE))
-	{
-		DisplayedClients.RemoveAt(Index);
-		ClientWidgets.Remove(ClientInfo.ClientEndpointId);
-		TileView->RequestListRefresh();
-	}
-}
-
-void UE::MultiUserServer::SConcertClientBrowser::UpdateClientInfo(const FConcertSessionClientInfo& ClientInfo)
-{
-	if (const TSharedPtr<SConcertClientBrowserItem>* Widget = ClientWidgets.Find(ClientInfo.ClientEndpointId))
-	{
-		Widget->Get()->OnClientInfoChanged();
-	}
+	TileView->RequestListRefresh();
 }
 
 TSharedRef<SWidget> UE::MultiUserServer::SConcertClientBrowser::SConcertClientBrowser::MakeSessionOption()
@@ -216,6 +265,21 @@ TSharedRef<SWidget> UE::MultiUserServer::SConcertClientBrowser::SConcertClientBr
 		NAME_None,
 		EUserInterfaceActionType::Check
 	);
+	MenuBuilder.AddMenuEntry(
+		LOCTEXT("SessionlessEndpoints.Label", "Show Sessionless clients"),
+		LOCTEXT("SessionlessEndpoints.Tooltip", "Whether to show clients that are only discovering available sessions"),
+		FSlateIcon(),
+		FUIAction(
+			FExecuteAction::CreateLambda([this]()
+			{
+				bShowSessionlessClients = !bShowSessionlessClients;
+				UpdateTileViewFromAllowedSessions();
+			}),
+			FCanExecuteAction::CreateLambda([] { return true; }),
+			FIsActionChecked::CreateLambda([this](){ return bShowSessionlessClients; })),
+		NAME_None,
+		EUserInterfaceActionType::Check
+	);
 	MenuBuilder.AddSeparator();
 	for (const FGuid& SessionId : BrowserModel->GetSessions())
 	{
@@ -258,25 +322,22 @@ FText UE::MultiUserServer::SConcertClientBrowser::GetErrorMessageText() const
 		return LOCTEXT("NoSessionsSelected", "All live sessions filtered out");
 	}
 
-	const bool bAtLeastOneClient = Algo::AnyOf(AllowedSessions, [this](const FGuid& SessionId)
-	{
-		return BrowserModel->GetSessionClients(SessionId).Num() > 0;
-	});
+	const bool bAtLeastOneClient = BrowserModel->GetItems().Num() > 0;
 	if (bAtLeastOneClient)
 	{
 		return LOCTEXT("AllFilteredOut", "All results have been filtered. Try changing your active filters above.");
 	}
 	
-	return LOCTEXT("NoClients", "No clients connected to selected sessions");
+	return LOCTEXT("NoClients", "No known clients");
 }
 
 TSharedRef<ITableRow> UE::MultiUserServer::SConcertClientBrowser::MakeTileViewWidget(TSharedPtr<FClientBrowserItem> ClientItem, const TSharedRef<STableViewBase>& OwnerTable)
 {
-	const FClientEndpointId ClientId = ClientItem->ClientInfo.ClientEndpointId;
-	TSharedPtr<SConcertClientBrowserItem>* ExistingWidget = ClientWidgets.Find(ClientId);
+	const FMessagingNodeId MessagingNodeId = ClientItem->MessageNodeId;
+	const TSharedPtr<SConcertClientBrowserItem>* ExistingWidget = ClientWidgets.Find(MessagingNodeId);
 	if (!ExistingWidget)
 	{
-		ClientWidgets.Add(ClientItem->ClientInfo.ClientEndpointId, SNew(SConcertClientBrowserItem, ClientItem.ToSharedRef(), StatisticsModel.ToSharedRef())
+		ClientWidgets.Add(MessagingNodeId, SNew(SConcertClientBrowserItem, ClientItem.ToSharedRef(), StatisticsModel.ToSharedRef())
 			 .HighlightText(HighlightText)
 			 );
 	}
@@ -286,7 +347,7 @@ TSharedRef<ITableRow> UE::MultiUserServer::SConcertClientBrowser::MakeTileViewWi
 		.Style(FConcertServerStyle::Get(), "Concert.Clients.TileTableRow")
 		.Content()
 		[
-			ClientWidgets[ClientId].ToSharedRef()
+			ClientWidgets[MessagingNodeId].ToSharedRef()
 		];
 }
 
@@ -318,7 +379,7 @@ TSharedPtr<SWidget> UE::MultiUserServer::SConcertClientBrowser::OnGetContextMenu
 
 void UE::MultiUserServer::SConcertClientBrowser::OnListMouseButtonDoubleClick(TSharedPtr<FClientBrowserItem> ClientItem)
 {
-	OnClientDoubleClicked.ExecuteIfBound(ClientItem->ClientInfo.ClientEndpointId);
+	OnClientDoubleClicked.ExecuteIfBound(ClientItem->MessageNodeId);
 }
 
 void UE::MultiUserServer::SConcertClientBrowser::SConcertClientBrowser::AllowAllSessions()
@@ -338,25 +399,30 @@ void UE::MultiUserServer::SConcertClientBrowser::UpdateTileViewFromAllowedSessio
 	bShowAllSessions = AllowedSessions.Num() == BrowserModel->GetSessions().Num();
 	DisplayedClients.Empty();
 	
-	for (const FGuid& AllowedSession : AllowedSessions)
+	for (const TSharedPtr<FClientBrowserItem>& Item : BrowserModel->GetItems())
 	{
-		for (const FConcertSessionClientInfo& ClientInfo : BrowserModel->GetSessionClients(AllowedSession))
+		if (PassesFilter(Item))
 		{
-			if (SessionFilter->PassesFilter(ClientInfo))
-			{
-				DisplayedClients.Add(MakeShared<FClientBrowserItem>(ClientInfo, BrowserModel->GetClientAddress(ClientInfo.ClientEndpointId)));
-			}
+			DisplayedClients.Add(Item);
 		}
 	}
 	TileView->RequestListRefresh();
 }
 
-void UE::MultiUserServer::SConcertClientBrowser::GenerateSearchTerms(const FConcertSessionClientInfo& ClientInfo, TArray<FString>& SearchTerms) const
+bool UE::MultiUserServer::SConcertClientBrowser::PassesFilter(const TSharedPtr<FClientBrowserItem>& Client) const
 {
-	if (const TSharedPtr<SConcertClientBrowserItem>* ItemWidget = ClientWidgets.Find(ClientInfo.ClientEndpointId))
+	const bool bIsSessionAllowed = Client->CurrentSession.IsSet() && AllowedSessions.Contains(*Client->CurrentSession);
+	const bool bAllowBaseOnSession = bIsSessionAllowed || (bShowSessionlessClients && !Client->CurrentSession.IsSet());
+	return bAllowBaseOnSession && SessionFilter->PassesFilter(Client);
+}
+
+void UE::MultiUserServer::SConcertClientBrowser::GenerateSearchTerms(const TSharedPtr<FClientBrowserItem>& Client, TArray<FString>& SearchTerms) const
+{
+	if (const TSharedPtr<SConcertClientBrowserItem>* ItemWidget = ClientWidgets.Find(Client->MessageNodeId))
 	{
 		ItemWidget->Get()->AppendSearchTerms(SearchTerms);
 	}
 }
 
+
 #undef LOCTEXT_NAMESPACE
diff --git a/Engine/Plugins/Developer/Concert/ConcertApp/MultiUserServer/Source/MultiUserServer/Private/Widgets/Clients/Browser/SConcertClientBrowser.h b/Engine/Plugins/Developer/Concert/ConcertApp/MultiUserServer/Source/MultiUserServer/Private/Widgets/Clients/Browser/SConcertClientBrowser.h
--- a/Engine/Plugins/Developer/Concert/ConcertApp/MultiUserServer/Source/MultiUserServer/Private/Widgets/Clients/Browser/SConcertClientBrowser.h
+++ b/Engine/Plugins/Developer/Concert/ConcertApp/MultiUserServer/Source/MultiUserServer/Private/Widgets/Clients/Browser/SConcertClientBrowser.h
@@ -5,6 +5,7 @@
 #include "CoreMinimal.h"
 #include "SConcertClientBrowserItem.h"
 #include "Misc/TextFilter.h"
+#include "Models/IClientBrowserModel.h"
 #include "Widgets/DeclarativeSyntaxSupport.h"
 #include "Widgets/SCompoundWidget.h"
 
@@ -46,8 +47,8 @@ namespace UE::MultiUserServer
 	private:
 
 		using FSessionId = FGuid;
-		using FClientEndpointId = FGuid;
-		using FClientTextFilter = TTextFilter<const FConcertSessionClientInfo&>;
+		using FMessagingNodeId = FGuid;
+		using FClientTextFilter = TTextFilter<const TSharedPtr<FClientBrowserItem>&>;
 
 		/** Retrieves clients and live sessions */
 		TSharedPtr<IClientBrowserModel> BrowserModel;
@@ -58,10 +59,12 @@ namespace UE::MultiUserServer
 		TSet<FSessionId> AllowedSessions;
 		/** Should all sessions be shown? */
 		bool bShowAllSessions = true;
+		/** Should admin endpoints be shown? */
+		bool bShowSessionlessClients = true;
 		/** Keeps widgets alive even when not displayed in the view - otherwise the graph will lose the historic data. */
-		TMap<FClientEndpointId, TSharedPtr<SConcertClientBrowserItem>> ClientWidgets;
+		TMap<FMessagingNodeId, TSharedPtr<SConcertClientBrowserItem>> ClientWidgets;
 
-		/** Source array for TileView */
+		/** Source array for TileView - fitlered version of IClientBrowserModel::GetItems */
 		TArray<TSharedPtr<FClientBrowserItem>> DisplayedClients;
 		/** Visualizes all the items */
 		TSharedPtr<STileView<TSharedPtr<FClientBrowserItem>>> TileView;
@@ -73,15 +76,13 @@ namespace UE::MultiUserServer
 		FOnClientDoubleClicked OnClientDoubleClicked;
 
 		TSharedRef<SWidget> CreateSearchArea(const FArguments& InArgs);
+		TSharedRef<SWidget> CreateKeepDisconnectedClients();
 		TSharedRef<SWidget> CreateTileView();
 
 		// Model events
 		void OnSessionCreated(const FGuid& SessionId);
 		void OnSessionDestroyed(const FGuid& SessionId);
-		void OnClientListChanged(const FGuid& SessionId, EConcertClientStatus UpdateType, const FConcertSessionClientInfo& ClientInfo);
-		
-		void RemoveClient(const FConcertSessionClientInfo& ClientInfo);
-		void UpdateClientInfo(const FConcertSessionClientInfo& ClientInfo);
+		void OnClientListChanged(TSharedPtr<FClientBrowserItem> Item, IClientBrowserModel::EClientUpdateType UpdateType);
 
 		// Combo button
 		TSharedRef<SWidget> MakeSessionOption();
@@ -97,7 +98,8 @@ namespace UE::MultiUserServer
 		void AllowAllSessions();
 		void DisallowAllSessions();
 		void UpdateTileViewFromAllowedSessions();
-		void GenerateSearchTerms(const FConcertSessionClientInfo& ClientInfo, TArray<FString>& SearchTerms) const;
+		bool PassesFilter(const TSharedPtr<FClientBrowserItem>& Client) const;
+		void GenerateSearchTerms(const TSharedPtr<FClientBrowserItem>& Client, TArray<FString>& SearchTerms) const;
 	};
 }
 
diff --git a/Engine/Plugins/Developer/Concert/ConcertApp/MultiUserServer/Source/MultiUserServer/Private/Widgets/Clients/Browser/SConcertClientBrowserItem.cpp b/Engine/Plugins/Developer/Concert/ConcertApp/MultiUserServer/Source/MultiUserServer/Private/Widgets/Clients/Browser/SConcertClientBrowserItem.cpp
--- a/Engine/Plugins/Developer/Concert/ConcertApp/MultiUserServer/Source/MultiUserServer/Private/Widgets/Clients/Browser/SConcertClientBrowserItem.cpp
+++ b/Engine/Plugins/Developer/Concert/ConcertApp/MultiUserServer/Source/MultiUserServer/Private/Widgets/Clients/Browser/SConcertClientBrowserItem.cpp
@@ -12,6 +12,8 @@
 #include "Widgets/Layout/SScaleBox.h"
 #include "Widgets/Text/STextBlock.h"
 
+#define LOCTEXT_NAMESPACE "UnrealMultiUserUI.SConcertClientBrowserItem"
+
 void UE::MultiUserServer::SConcertClientBrowserItem::Construct(const FArguments& InArgs, TSharedRef<FClientBrowserItem> InClientItem, TSharedRef<IClientNetworkStatisticsModel> InStatModel)
 {
 	Item = MoveTemp(InClientItem);
@@ -65,13 +67,21 @@ void UE::MultiUserServer::SConcertClientBrowserItem::Construct(const FArguments&
 			]
 		]
 	];
-
-	OnClientInfoChanged();
 }
 
-void UE::MultiUserServer::SConcertClientBrowserItem::OnClientInfoChanged()
+FString UE::MultiUserServer::SConcertClientBrowserItem::GetClientDisplayName() const
 {
-	ClientName->SetText(FText::FromString(Item->ClientInfo.ClientInfo.DisplayName));
+	const TOptional<FConcertClientInfo> ClientInfo = Item->GetClientInfo.Execute();
+	if (ClientInfo)
+	{
+		return ClientInfo->DisplayName;
+	}
+	
+	const FString NodeIdAsString = Item->MessageNodeId.ToString(EGuidFormats::DigitsWithHyphens);
+	int32 Index = INDEX_NONE;
+	const bool bFound = NodeIdAsString.FindChar('-', Index);
+	// Avoid making the name too long by only showing the first few digits of the node ID
+	return FString::Printf(TEXT("Admin (%s)"), bFound ? *NodeIdAsString.Left(Index) : *NodeIdAsString);
 }
 
 void UE::MultiUserServer::SConcertClientBrowserItem::AppendSearchTerms(TArray<FString>& SearchTerms) const
@@ -85,7 +95,24 @@ TSharedRef<SWidget> UE::MultiUserServer::SConcertClientBrowserItem::CreateHeader
 {
 	return SAssignNew(ClientName, STextBlock)
 		.Font(FConcertServerStyle::Get().GetFontStyle("Concert.Clients.ClientNameTileFont"))
-		.Text_Lambda([this](){ return FText::FromString(Item->ClientInfo.ClientInfo.DisplayName);} )
+		.Text_Lambda([this](){ return FText::FromString(GetClientDisplayName()); })
+		.ToolTipText_Lambda([this]()
+		{
+			const TOptional<FConcertClientInfo> ClientInfo = Item->GetClientInfo.Execute();
+			if (!ClientInfo)
+			{
+				return FText::Format(
+					LOCTEXT("Name.NotAvailable.TooltipFmt", "This client's display information becomes available after joining a session.\nNodeID: {0}\nAddress ID: {1}"),
+					FText::FromString(Item->MessageNodeId.ToString(EGuidFormats::DigitsWithHyphens)),
+					FText::FromString(Item->ClientAddress.ToString())
+					);
+			}
+			return FText::Format(
+					LOCTEXT("Name.Available.TooltipFmt", "NodeID: {0}\nAddress ID: {1}"),
+					FText::FromString(Item->MessageNodeId.ToString(EGuidFormats::DigitsWithHyphens)),
+					FText::FromString(Item->ClientAddress.ToString())
+					);
+		})
 		.HighlightText_Lambda([this](){ return *HighlightText; })
 		.ColorAndOpacity(FColor::White);
 }
@@ -139,3 +166,5 @@ const FSlateBrush* UE::MultiUserServer::SConcertClientBrowserItem::GetBackground
 	}
 	return FConcertServerStyle::Get().GetBrush("Concert.Clients.ThumbnailAreaBackground");
 }
+
+#undef LOCTEXT_NAMESPACE
\ No newline at end of file
diff --git a/Engine/Plugins/Developer/Concert/ConcertApp/MultiUserServer/Source/MultiUserServer/Private/Widgets/Clients/Browser/SConcertClientBrowserItem.h b/Engine/Plugins/Developer/Concert/ConcertApp/MultiUserServer/Source/MultiUserServer/Private/Widgets/Clients/Browser/SConcertClientBrowserItem.h
--- a/Engine/Plugins/Developer/Concert/ConcertApp/MultiUserServer/Source/MultiUserServer/Private/Widgets/Clients/Browser/SConcertClientBrowserItem.h
+++ b/Engine/Plugins/Developer/Concert/ConcertApp/MultiUserServer/Source/MultiUserServer/Private/Widgets/Clients/Browser/SConcertClientBrowserItem.h
@@ -25,9 +25,8 @@ namespace UE::MultiUserServer
 
 		void Construct(const FArguments& InArgs, TSharedRef<FClientBrowserItem> InClientItem, TSharedRef<IClientNetworkStatisticsModel> InStatModel);
 
-		/** Called when the underlying client data was changed */
-		void OnClientInfoChanged();
-
+		FString GetClientDisplayName() const;
+		
 		/** Appends the statistics to the search terms */
 		void AppendSearchTerms(TArray<FString>& SearchTerms) const;
 
diff --git a/Engine/Plugins/Developer/Concert/ConcertApp/MultiUserServer/Source/MultiUserServer/Private/Widgets/Clients/Logging/Filter/ConcertFrontendLogFilter_Client.cpp b/Engine/Plugins/Developer/Concert/ConcertApp/MultiUserServer/Source/MultiUserServer/Private/Widgets/Clients/Logging/Filter/ConcertFrontendLogFilter_Client.cpp
--- a/Engine/Plugins/Developer/Concert/ConcertApp/MultiUserServer/Source/MultiUserServer/Private/Widgets/Clients/Logging/Filter/ConcertFrontendLogFilter_Client.cpp
+++ b/Engine/Plugins/Developer/Concert/ConcertApp/MultiUserServer/Source/MultiUserServer/Private/Widgets/Clients/Logging/Filter/ConcertFrontendLogFilter_Client.cpp
@@ -4,18 +4,14 @@
 
 #include "ConcertTransportEvents.h"
 #include "Widgets/Clients/Logging/ConcertLogEntry.h"
+#include "Widgets/Clients/Util/EndpointToUserNameCache.h"
 
-bool FConcertLogFilter_Client::PassesFilter(const FConcertLogEntry& InItem) const
+namespace UE::MultiUserServer::Filters
 {
-	return AllowedClientEndpointIds.Contains(InItem.Log.OriginEndpointId)
-		|| AllowedClientEndpointIds.Contains(InItem.Log.DestinationEndpointId);
-}
-
-void FConcertLogFilter_Client::AllowOnly(const FGuid& ClientEndpointId)
-{
-	if (AllowedClientEndpointIds.Num() > 1 || !AllowedClientEndpointIds.Contains(ClientEndpointId))
+	bool FConcertLogFilter_Client::PassesFilter(const FConcertLogEntry& InItem) const
 	{
-		AllowedClientEndpointIds = { ClientEndpointId };
-		OnChanged().Broadcast();
+		return AllowedClientMessagingNodeId == EndpointCache->TranslateEndpointIdToNodeId(InItem.Log.OriginEndpointId)
+			|| AllowedClientMessagingNodeId == EndpointCache->TranslateEndpointIdToNodeId(InItem.Log.DestinationEndpointId);
 	}
 }
+
diff --git a/Engine/Plugins/Developer/Concert/ConcertApp/MultiUserServer/Source/MultiUserServer/Private/Widgets/Clients/Logging/Filter/ConcertFrontendLogFilter_Client.h b/Engine/Plugins/Developer/Concert/ConcertApp/MultiUserServer/Source/MultiUserServer/Private/Widgets/Clients/Logging/Filter/ConcertFrontendLogFilter_Client.h
--- a/Engine/Plugins/Developer/Concert/ConcertApp/MultiUserServer/Source/MultiUserServer/Private/Widgets/Clients/Logging/Filter/ConcertFrontendLogFilter_Client.h
+++ b/Engine/Plugins/Developer/Concert/ConcertApp/MultiUserServer/Source/MultiUserServer/Private/Widgets/Clients/Logging/Filter/ConcertFrontendLogFilter_Client.h
@@ -5,24 +5,30 @@
 #include "CoreMinimal.h"
 #include "Widgets/Clients/Logging/Filter/ConcertFrontendLogFilter.h"
 
-/** Only allows messages from the given clients */
-class FConcertLogFilter_Client : public FConcertLogFilter
+class FEndpointToUserNameCache;
+
+namespace UE::MultiUserServer::Filters
 {
-public:
+	/** Only allows messages from the given clients */
+	class FConcertLogFilter_Client : public FConcertLogFilter
+	{
+	public:
 
-	FConcertLogFilter_Client() = default;
-	explicit FConcertLogFilter_Client(const FGuid& SingleAllowedId)
-		: AllowedClientEndpointIds({ SingleAllowedId })
-	{}
+		explicit FConcertLogFilter_Client(const FGuid& SingleAllowedId, TSharedRef<FEndpointToUserNameCache> EndpointCache)
+			: AllowedClientMessagingNodeId(SingleAllowedId)
+			, EndpointCache(MoveTemp(EndpointCache))
+		{}
 
-	//~ Begin FConcertLogFilter Interface
-	virtual bool PassesFilter(const FConcertLogEntry& InItem) const override;
-	//~ End FConcertLogFilter Interface
+		//~ Begin FConcertLogFilter Interface
+		virtual bool PassesFilter(const FConcertLogEntry& InItem) const override;
+		//~ End FConcertLogFilter Interface
 
-	void AllowOnly(const FGuid& ClientEndpointId);
+	private:
 
-private:
+		/** Messages to and from the following client endpoint ID are allowed */
+		FGuid AllowedClientMessagingNodeId;
 
-	/** Messages to and from the following client endpoint IDs are allowed */
-	TSet<FGuid> AllowedClientEndpointIds;
-};
+		/** Translates Concert endpoint Ids to messaging node Ids */
+		TSharedRef<FEndpointToUserNameCache> EndpointCache;
+	};
+}
diff --git a/Engine/Plugins/Developer/Concert/ConcertApp/MultiUserServer/Source/MultiUserServer/Private/Widgets/Clients/Logging/Filter/ConcertLogFilter_FrontendRoot.cpp b/Engine/Plugins/Developer/Concert/ConcertApp/MultiUserServer/Source/MultiUserServer/Private/Widgets/Clients/Logging/Filter/ConcertLogFilter_FrontendRoot.cpp
--- a/Engine/Plugins/Developer/Concert/ConcertApp/MultiUserServer/Source/MultiUserServer/Private/Widgets/Clients/Logging/Filter/ConcertLogFilter_FrontendRoot.cpp
+++ b/Engine/Plugins/Developer/Concert/ConcertApp/MultiUserServer/Source/MultiUserServer/Private/Widgets/Clients/Logging/Filter/ConcertLogFilter_FrontendRoot.cpp
@@ -110,11 +110,11 @@ namespace UE::MultiUserServer
 			);
 	}
 
-	TSharedRef<FConcertLogFilter_FrontendRoot> MakeClientLogFilter(TSharedRef<FConcertLogTokenizer> Tokenizer, const FGuid& ClientEndpointId)
+	TSharedRef<FConcertLogFilter_FrontendRoot> MakeClientLogFilter(TSharedRef<FConcertLogTokenizer> Tokenizer, const FGuid& ClientMessageNodeId, const TSharedRef<FEndpointToUserNameCache>& EndpointCache)
 	{
 		const TArray<TSharedRef<FConcertFrontendLogFilter>> CommonFilters = Private::CreateCommonFilters();
 		const TArray<TSharedRef<FConcertLogFilter>> NonVisuals = {
-			MakeShared<FConcertLogFilter_Client>(ClientEndpointId)
+			MakeShared<Filters::FConcertLogFilter_Client>(ClientMessageNodeId, EndpointCache)
 		};
 		return MakeShared<FConcertLogFilter_FrontendRoot>(
 			MoveTemp(Tokenizer),
diff --git a/Engine/Plugins/Developer/Concert/ConcertApp/MultiUserServer/Source/MultiUserServer/Private/Widgets/Clients/Logging/Filter/ConcertLogFilter_FrontendRoot.h b/Engine/Plugins/Developer/Concert/ConcertApp/MultiUserServer/Source/MultiUserServer/Private/Widgets/Clients/Logging/Filter/ConcertLogFilter_FrontendRoot.h
--- a/Engine/Plugins/Developer/Concert/ConcertApp/MultiUserServer/Source/MultiUserServer/Private/Widgets/Clients/Logging/Filter/ConcertLogFilter_FrontendRoot.h
+++ b/Engine/Plugins/Developer/Concert/ConcertApp/MultiUserServer/Source/MultiUserServer/Private/Widgets/Clients/Logging/Filter/ConcertLogFilter_FrontendRoot.h
@@ -8,6 +8,7 @@
 class FConcertLogTokenizer;
 class FConcertFrontendLogFilter;
 class FConcertFrontendLogFilter_TextSearch;
+class FEndpointToUserNameCache;
 class SWidget;
 
 /** A filter that contains multiple UI filters */
@@ -48,6 +49,11 @@ namespace UE::MultiUserServer
 	/** Creates a filter for the global filter log window */
 	TSharedRef<FConcertLogFilter_FrontendRoot> MakeGlobalLogFilter(TSharedRef<FConcertLogTokenizer> Tokenizer);
 
-	/** Creates a filter for a filter log window intended for a client */
-	TSharedRef<FConcertLogFilter_FrontendRoot> MakeClientLogFilter(TSharedRef<FConcertLogTokenizer> Tokenizer, const FGuid& ClientEndpointId);
+	/**
+	 * Creates a filter for a filter log window intended for a client
+	 * @param Tokenizer Used for text search
+	 * @param ClientMessageNodeId The Id of this client's messaging node - this is used to filter messages involving this client
+	 * @param EndpointCache Needed to filter messages involving this client - converts Concert endpoint Ids to the message node Id
+	 */
+	TSharedRef<FConcertLogFilter_FrontendRoot> MakeClientLogFilter(TSharedRef<FConcertLogTokenizer> Tokenizer, const FGuid& ClientMessageNodeId, const TSharedRef<FEndpointToUserNameCache>& EndpointCache);
 }
diff --git a/Engine/Plugins/Developer/Concert/ConcertApp/MultiUserServer/Source/MultiUserServer/Private/Widgets/Clients/Logging/SConcertTransportLog.cpp b/Engine/Plugins/Developer/Concert/ConcertApp/MultiUserServer/Source/MultiUserServer/Private/Widgets/Clients/Logging/SConcertTransportLog.cpp
--- a/Engine/Plugins/Developer/Concert/ConcertApp/MultiUserServer/Source/MultiUserServer/Private/Widgets/Clients/Logging/SConcertTransportLog.cpp
+++ b/Engine/Plugins/Developer/Concert/ConcertApp/MultiUserServer/Source/MultiUserServer/Private/Widgets/Clients/Logging/SConcertTransportLog.cpp
@@ -14,14 +14,14 @@
 
 #include "Framework/MultiBox/MultiBoxBuilder.h"
 #include "Styling/AppStyle.h"
-#include "Util/EndpointToUserNameCache.h"
 #include "Widgets/Layout/SBorder.h"
 #include "Widgets/SBoxPanel.h"
 #include "Widgets/SNullWidget.h"
+#include "Widgets/Clients/Util/EndpointToUserNameCache.h"
 #include "Widgets/Views/SHeaderRow.h"
 #include "Widgets/Views/SListView.h"
 
-#define LOCTEXT_NAMESPACE "UnrealMultiUserUI"
+#define LOCTEXT_NAMESPACE "UnrealMultiUserUI.SConcertTransportLog"
 
 const FName SConcertTransportLog::FirstColumnId("AvatarColourColumnId");
 
@@ -264,7 +264,7 @@ void SConcertTransportLog::ExtendViewOptions(FMenuBuilder& MenuBuilder)
 			EUserInterfaceActionType::Button
 		);
 	MenuBuilder.AddMenuEntry(
-			LOCTEXT("SelectAll", "Hide all"),
+			LOCTEXT("HideAll", "Hide all"),
 			FText::GetEmpty(),
 			FSlateIcon(),
 			FUIAction(
diff --git a/Engine/Plugins/Developer/Concert/ConcertApp/MultiUserServer/Source/MultiUserServer/Private/Widgets/Clients/Logging/Util/ConcertLogTokenizer.cpp b/Engine/Plugins/Developer/Concert/ConcertApp/MultiUserServer/Source/MultiUserServer/Private/Widgets/Clients/Logging/Util/ConcertLogTokenizer.cpp
--- a/Engine/Plugins/Developer/Concert/ConcertApp/MultiUserServer/Source/MultiUserServer/Private/Widgets/Clients/Logging/Util/ConcertLogTokenizer.cpp
+++ b/Engine/Plugins/Developer/Concert/ConcertApp/MultiUserServer/Source/MultiUserServer/Private/Widgets/Clients/Logging/Util/ConcertLogTokenizer.cpp
@@ -4,10 +4,10 @@
 
 #include "ConcertFrontendUtils.h"
 #include "ConcertTransportEvents.h"
-#include "EndpointToUserNameCache.h"
 #include "MessageTypeUtils.h"
 #include "Math/UnitConversion.h"
 #include "Settings/ConcertTransportLogSettings.h"
+#include "Widgets/Clients/Util/EndpointToUserNameCache.h"
 
 FConcertLogTokenizer::FConcertLogTokenizer(TSharedRef<FEndpointToUserNameCache> EndpointInfoGetter)
 	: EndpointInfoGetter(EndpointInfoGetter)
diff --git a/Engine/Plugins/Developer/Concert/ConcertApp/MultiUserServer/Source/MultiUserServer/Private/Widgets/Clients/Logging/Util/EndpointToUserNameCache.cpp b/Engine/Plugins/Developer/Concert/ConcertApp/MultiUserServer/Source/MultiUserServer/Private/Widgets/Clients/Logging/Util/EndpointToUserNameCache.cpp
deleted file mode 100644
diff --git a/Engine/Plugins/Developer/Concert/ConcertApp/MultiUserServer/Source/MultiUserServer/Private/Widgets/Clients/Logging/Util/EndpointToUserNameCache.h b/Engine/Plugins/Developer/Concert/ConcertApp/MultiUserServer/Source/MultiUserServer/Private/Widgets/Clients/Logging/Util/EndpointToUserNameCache.h
deleted file mode 100644
diff --git a/Engine/Plugins/Developer/Concert/ConcertApp/MultiUserServer/Source/MultiUserServer/Private/Widgets/Clients/SConcertClientsTabView.cpp b/Engine/Plugins/Developer/Concert/ConcertApp/MultiUserServer/Source/MultiUserServer/Private/Widgets/Clients/SConcertClientsTabView.cpp
--- a/Engine/Plugins/Developer/Concert/ConcertApp/MultiUserServer/Source/MultiUserServer/Private/Widgets/Clients/SConcertClientsTabView.cpp
+++ b/Engine/Plugins/Developer/Concert/ConcertApp/MultiUserServer/Source/MultiUserServer/Private/Widgets/Clients/SConcertClientsTabView.cpp
@@ -2,20 +2,19 @@
 
 #include "SConcertClientsTabView.h"
 
-#include "ConcertClientsTabController.h"
 #include "ConcertUtil.h"
 #include "IConcertServer.h"
 #include "IConcertSyncServer.h"
 #include "Logging/Filter/ConcertLogFilter_FrontendRoot.h"
 #include "Logging/Source/GlobalLogSource.h"
 #include "Logging/Util/ConcertLogTokenizer.h"
+#include "Util/EndpointToUserNameCache.h"
 #include "Widgets/Clients/Browser/SConcertClientBrowser.h"
 #include "Widgets/Clients/Browser/Models/ClientBrowserModel.h"
 #include "Widgets/Clients/Browser/Models/ClientNetworkStatisticsModel.h"
 #include "Widgets/Clients/Logging/SConcertTransportLog.h"
 
 #include "Framework/Docking/TabManager.h"
-#include "Logging/Util/EndpointToUserNameCache.h"
 #include "Widgets/Docking/SDockTab.h"
 #include "Widgets/Images/SImage.h"
 #include "Widgets/Input/SButton.h"
@@ -51,24 +50,29 @@ void SConcertClientsTabView::ShowConnectedClients(const FGuid& SessionId) const
 	ClientBrowser->ShowOnlyClientsFromSession(SessionId);
 }
 
-void SConcertClientsTabView::OpenClientLogTab(const FGuid& ClientEndpointId) const
+void SConcertClientsTabView::OpenClientLogTab(const FGuid& ClientMessageNodeId) const
 {
-
-	const FName TabId = *ClientEndpointId.ToString();
+	const FName TabId = *ClientMessageNodeId.ToString();
 	if (const TSharedPtr<SDockTab> ExistingTab = GetTabManager()->FindExistingLiveTab(FTabId(TabId)))
 	{
 		GetTabManager()->DrawAttention(ExistingTab.ToSharedRef());
 	}
-	else if (const TOptional<FConcertSessionClientInfo> ClientInfo = ConcertUtil::GetConnectedClientInfo(*Server->GetConcertServer(), ClientEndpointId)
-		; ensure(ClientInfo))
+	else
 	{
+		const TOptional<FConcertClientInfo> ClientInfo = ClientInfoCache->GetClientInfoFromNodeId(ClientMessageNodeId);
 		const TSharedRef<SDockTab> NewTab = SNew(SDockTab)
-			.Label(FText::Format(LOCTEXT("ClientTabFmt", "{0} Log"), FText::FromString(ClientInfo->ClientInfo.DisplayName)))
-			.ToolTipText(FText::Format(LOCTEXT("ClientTabTooltipFmt", "Logs all networked requests originating or going to  client {0} (EndpointId = {1})"), FText::FromString(ClientInfo->ClientInfo.DisplayName), FText::FromString(ClientEndpointId.ToString())))
+			.Label_Lambda([this, ClientMessageNodeId]()
+			{
+				const TOptional<FConcertClientInfo> ClientInfo = ClientInfoCache->GetClientInfoFromNodeId(ClientMessageNodeId);
+				return ClientInfo
+					? FText::Format(LOCTEXT("ClientTabFmt", "{0} Log"), FText::FromString(ClientInfo->DisplayName))
+					: FText::FromString(ClientMessageNodeId.ToString(EGuidFormats::DigitsWithHyphens));
+			})
+			.ToolTipText(FText::Format(LOCTEXT("ClientTabTooltipFmt", "Logs all networked requests originating or going to client {0} (NodeId = {1})"), ClientInfo ? FText::FromString(ClientInfo->DisplayName) : FText::GetEmpty(), FText::FromString(ClientMessageNodeId.ToString())))
 			.TabRole(PanelTab)
 			[
 				SNew(SConcertTransportLog, LogBuffer.ToSharedRef(), ClientInfoCache.ToSharedRef(), LogTokenizer.ToSharedRef())
-				.Filter(UE::MultiUserServer::MakeClientLogFilter(LogTokenizer.ToSharedRef(), ClientEndpointId))
+				.Filter(UE::MultiUserServer::MakeClientLogFilter(LogTokenizer.ToSharedRef(), ClientMessageNodeId, ClientInfoCache.ToSharedRef()))
 			]; 
 
 		// We need a tab to place the client tab next to
@@ -144,7 +148,7 @@ TSharedRef<SDockTab> SConcertClientsTabView::SpawnClientBrowserTab(const FSpawnT
 		.TabRole(PanelTab)
 		[
 			SAssignNew(ClientBrowser, UE::MultiUserServer::SConcertClientBrowser,
-				MakeShared<UE::MultiUserServer::FClientBrowserModel>(Server->GetConcertServer()),
+				MakeShared<UE::MultiUserServer::FClientBrowserModel>(Server->GetConcertServer(), ClientInfoCache.ToSharedRef()),
 				MakeShared<UE::MultiUserServer::FClientNetworkStatisticsModel>())
 			.RightOfSearch()
 			[
@@ -204,18 +208,4 @@ TSharedRef<SWidget> SConcertClientsTabView::CreateOpenGlobalLogButton() const
 		];
 }
 
-TOptional<FConcertClientInfo> SConcertClientsTabView::GetClientInfo(const FGuid& EndpointId) const
-{
-	if (EndpointId == Server->GetConcertServer()->GetServerInfo().AdminEndpointId)
-	{
-		return {};
-	}
-
-	if (const TOptional<FConcertSessionClientInfo> ClientInfo = ConcertUtil::GetConnectedClientInfo(Server->GetConcertServer().Get(), EndpointId))
-	{
-		return ClientInfo->ClientInfo;
-	}
-	return {};
-}
-
 #undef LOCTEXT_NAMESPACE
diff --git a/Engine/Plugins/Developer/Concert/ConcertApp/MultiUserServer/Source/MultiUserServer/Private/Widgets/Clients/SConcertClientsTabView.h b/Engine/Plugins/Developer/Concert/ConcertApp/MultiUserServer/Source/MultiUserServer/Private/Widgets/Clients/SConcertClientsTabView.h
--- a/Engine/Plugins/Developer/Concert/ConcertApp/MultiUserServer/Source/MultiUserServer/Private/Widgets/Clients/SConcertClientsTabView.h
+++ b/Engine/Plugins/Developer/Concert/ConcertApp/MultiUserServer/Source/MultiUserServer/Private/Widgets/Clients/SConcertClientsTabView.h
@@ -45,7 +45,7 @@ public:
 	void ShowConnectedClients(const FGuid& SessionId) const;
 	void OpenGlobalLogTab() const;
 	void CloseGlobalLogTab() const;
-	void OpenClientLogTab(const FGuid& ClientEndpointId) const;
+	void OpenClientLogTab(const FGuid& ClientMessageNodeId) const;
 	
 	bool IsGlobalLogOpen() const;
 	TSharedPtr<SDockTab> GetGlobalLogTab() const;
@@ -74,5 +74,4 @@ private:
 	TSharedRef<SDockTab> SpawnGlobalLogTab(const FSpawnTabArgs& InTabArgs);
 
 	TSharedRef<SWidget> CreateOpenGlobalLogButton() const;
-	TOptional<FConcertClientInfo> GetClientInfo(const FGuid& EndpointId) const;
 };
diff --git a/Engine/Plugins/Developer/Concert/ConcertApp/MultiUserServer/Source/MultiUserServer/Private/Widgets/Clients/Util/EndpointToUserNameCache.cpp b/Engine/Plugins/Developer/Concert/ConcertApp/MultiUserServer/Source/MultiUserServer/Private/Widgets/Clients/Util/EndpointToUserNameCache.cpp
new file mode 100644
--- /dev/null
+++ b/Engine/Plugins/Developer/Concert/ConcertApp/MultiUserServer/Source/MultiUserServer/Private/Widgets/Clients/Util/EndpointToUserNameCache.cpp
@@ -0,0 +1,172 @@
+// Copyright Epic Games, Inc. All Rights Reserved.
+
+#include "EndpointToUserNameCache.h"
+
+#include "ConcertServerEvents.h"
+#include "ConcertUtil.h"
+#include "IConcertServer.h"
+#include "IConcertSession.h"
+#include "INetworkMessagingExtension.h"
+#include "Algo/AnyOf.h"
+#include "Features/IModularFeatures.h"
+
+FEndpointToUserNameCache::FEndpointToUserNameCache(TSharedRef<IConcertServer> Server)
+	: Server(MoveTemp(Server))
+{
+	ConcertServerEvents::OnLiveSessionCreated().AddRaw(this, &FEndpointToUserNameCache::OnLiveSessionCreated);
+	ConcertServerEvents::OnLiveSessionDestroyed().AddRaw(this, &FEndpointToUserNameCache::OnLiveSessionDestroyed);
+
+	for (TSharedPtr<IConcertServerSession> LiveSession : Server->GetLiveSessions())
+	{
+		RegisterLiveSession(LiveSession.ToSharedRef());
+		for (const FConcertSessionClientInfo& ClientInfo : LiveSession->GetSessionClients())
+		{
+			CacheClientInfo(*LiveSession, ClientInfo);
+		}
+	}
+
+	Server->OnRemoteEndpointConnectionChanged().AddRaw(this, &FEndpointToUserNameCache::OnAdminEndpointConnectionChanged);
+}
+
+FEndpointToUserNameCache::~FEndpointToUserNameCache()
+{
+	ConcertServerEvents::OnLiveSessionCreated().RemoveAll(this);
+	ConcertServerEvents::OnLiveSessionDestroyed().RemoveAll(this);
+
+	for (TWeakPtr<IConcertServerSession> Session : SubscribedToSessions)
+	{
+		if (TSharedPtr<IConcertServerSession> PinnedSession = Session.Pin())
+		{
+			PinnedSession->OnSessionClientChanged().RemoveAll(this);
+		}
+	}
+}
+
+bool FEndpointToUserNameCache::IsServerEndpoint(const FGuid& EndpointId) const
+{
+	const bool bIsLocalSessionEndpoint = Algo::AnyOf(Server->GetLiveSessions(), [&EndpointId](const TSharedPtr<IConcertServerSession>& Session)
+	{
+		return Session->GetSessionInfo().ServerEndpointId == EndpointId;
+	});
+	return bIsLocalSessionEndpoint || Server->GetServerInfo().AdminEndpointId == EndpointId;
+}
+
+TOptional<FConcertClientInfo> FEndpointToUserNameCache::GetClientInfo(const FGuid& EndpointId) const
+{
+	// Already known or an old endpoint ID that is no longer valid (e.g. because client left session)
+	if (const FNodeEndpointId* CachedNodeId = CachedConcertEndpointToNodeEndpoints.Find(EndpointId))
+	{
+		if (const FConcertClientInfo* ClientData = CachedClientData.Find(*CachedNodeId))
+		{
+			return *ClientData;
+		}
+	}
+	
+	auto LookUpAddress = [this](const FMessageAddress& MessageAddress) -> TOptional<FConcertClientInfo>
+	{
+		const FNodeEndpointId NodeEndpointId = GetNodeIdFromMessagingBackend(MessageAddress);
+		if (const FConcertClientInfo* Data = CachedClientData.Find(NodeEndpointId)) 
+		{
+			return *Data;
+		}
+		return {};
+	};
+
+	// Is in any live session?
+	if (const TSharedPtr<IConcertServerSession> Session = ConcertUtil::GetLiveSessionClientConnectedTo(Server.Get(), EndpointId))
+	{
+		const FMessageAddress ConnectedClientAddress = Session->GetClientAddress(EndpointId);
+		return LookUpAddress(ConnectedClientAddress);
+	}
+
+	// Is somebody trying to discover sessions? Data will only be available after this client has joined at least one session
+	if (const FMessageAddress RemoteAdminAddress = Server->GetRemoteAddress(EndpointId); RemoteAdminAddress.IsValid())
+	{
+		return LookUpAddress(RemoteAdminAddress);
+	}
+	return {};
+}
+
+TOptional<FConcertClientInfo> FEndpointToUserNameCache::GetClientInfoFromNodeId(const FNodeEndpointId& EndpointId) const
+{
+	if (const FConcertClientInfo* ClientData = CachedClientData.Find(EndpointId))
+	{
+		return *ClientData;
+	}
+	return {};
+}
+
+TOptional<FEndpointToUserNameCache::FNodeEndpointId> FEndpointToUserNameCache::TranslateEndpointIdToNodeId(const FGuid& EndpointId) const
+{
+	const FNodeEndpointId* CachedNodeId = CachedConcertEndpointToNodeEndpoints.Find(EndpointId);
+	if (CachedNodeId)
+	{
+		return *CachedNodeId;
+	}
+	return {};
+}
+
+void FEndpointToUserNameCache::OnLiveSessionCreated(bool bSuccess, const IConcertServer& InServer, TSharedRef<IConcertServerSession> InLiveSession)
+{
+	RegisterLiveSession(InLiveSession);
+}
+
+void FEndpointToUserNameCache::OnLiveSessionDestroyed(const IConcertServer& InServer, TSharedRef<IConcertServerSession> InLiveSession)
+{
+	// No need to unsubscribe because the session will be destroyed anyways
+	SubscribedToSessions.Remove(InLiveSession);
+}
+
+void FEndpointToUserNameCache::OnClientInfoChanged(IConcertServerSession& Session, EConcertClientStatus ConnectionStatus, const FConcertSessionClientInfo& ClientInfo)
+{
+	switch (ConnectionStatus)
+	{
+	case EConcertClientStatus::Disconnected:
+		break;
+		
+	case EConcertClientStatus::Connected:
+	case EConcertClientStatus::Updated:
+		CacheClientInfo(Session, ClientInfo);
+		break;
+	default:
+		checkNoEntry();
+	}
+}
+
+void FEndpointToUserNameCache::OnAdminEndpointConnectionChanged(const FConcertEndpointContext& ConcertEndpointContext, EConcertRemoteEndpointConnection ConcertRemoteEndpointConnection)
+{
+	const FMessageAddress Address = Server->GetRemoteAddress(ConcertEndpointContext.EndpointId);
+	if (const FNodeEndpointId NodeEndpointId = GetNodeIdFromMessagingBackend(Address); ConcertRemoteEndpointConnection == EConcertRemoteEndpointConnection::Discovered && NodeEndpointId.IsValid())
+	{
+		CachedConcertEndpointToNodeEndpoints.Add(ConcertEndpointContext.EndpointId, NodeEndpointId);
+	}
+}
+
+void FEndpointToUserNameCache::RegisterLiveSession(const TSharedRef<IConcertServerSession>& InLiveSession)
+{
+	SubscribedToSessions.Add(InLiveSession);
+	InLiveSession->OnSessionClientChanged().AddRaw(this, &FEndpointToUserNameCache::OnClientInfoChanged);
+}
+
+void FEndpointToUserNameCache::CacheClientInfo(const IConcertServerSession& Session, const FConcertSessionClientInfo& ClientInfo)
+{
+	const FMessageAddress Address = Session.GetClientAddress(ClientInfo.ClientEndpointId);
+	if (const FNodeEndpointId NodeEndpointId = GetNodeIdFromMessagingBackend(Address); NodeEndpointId.IsValid())
+	{
+		CachedConcertEndpointToNodeEndpoints.Add(ClientInfo.ClientEndpointId, NodeEndpointId);
+		CachedClientData.Add(NodeEndpointId, ClientInfo.ClientInfo);
+	}
+}
+
+FEndpointToUserNameCache::FNodeEndpointId FEndpointToUserNameCache::GetNodeIdFromMessagingBackend(const FMessageAddress& MessageAddress) const
+{
+	IModularFeatures& ModularFeatures = IModularFeatures::Get();
+	const INetworkMessagingExtension* Extension = ModularFeatures.IsModularFeatureAvailable(INetworkMessagingExtension::ModularFeatureName)
+		? &ModularFeatures.GetModularFeature<INetworkMessagingExtension>(INetworkMessagingExtension::ModularFeatureName)
+		: nullptr;
+	if (Extension)
+	{
+		return Extension->GetNodeIdFromAddress(MessageAddress);
+	}
+	return {};
+}
diff --git a/Engine/Plugins/Developer/Concert/ConcertApp/MultiUserServer/Source/MultiUserServer/Private/Widgets/Clients/Util/EndpointToUserNameCache.h b/Engine/Plugins/Developer/Concert/ConcertApp/MultiUserServer/Source/MultiUserServer/Private/Widgets/Clients/Util/EndpointToUserNameCache.h
new file mode 100644
--- /dev/null
+++ b/Engine/Plugins/Developer/Concert/ConcertApp/MultiUserServer/Source/MultiUserServer/Private/Widgets/Clients/Util/EndpointToUserNameCache.h
@@ -0,0 +1,71 @@
+// Copyright Epic Games, Inc. All Rights Reserved.
+
+#pragma once
+
+#include "CoreMinimal.h"
+#include "ConcertMessageData.h"
+#include "IConcertEndpoint.h"
+#include "IMessageContext.h"
+
+class IConcertServerSession;
+class IConcertServer;
+
+enum class EConcertClientStatus : uint8;
+
+struct FConcertSessionClientInfo;
+
+/**
+ * Responsible for keeping hold of client info even when the client disconnects.
+ * Client info is removed server-side upon disconnect but logs need their display info.
+ */
+class FEndpointToUserNameCache
+{
+public:
+	
+	/** ID used by the messaging system - corresponds to an IP address. */
+	using FNodeEndpointId = FGuid;
+
+	FEndpointToUserNameCache(TSharedRef<IConcertServer> Server);
+	~FEndpointToUserNameCache();
+
+	bool IsServerEndpoint(const FGuid& EndpointId) const;
+	/** Gets the client info given from any endpoint ID (any ID of a FConcertRemoteEndpoint) */
+	TOptional<FConcertClientInfo> GetClientInfo(const FGuid& EndpointId) const;
+	/** Gets the client info given from any endpoint node ID (the ID from the messaging system node - corresponds to IP address). */
+	TOptional<FConcertClientInfo> GetClientInfoFromNodeId(const FNodeEndpointId& EndpointId) const;
+
+	/** Translates a Concert endpoint Id to the corresponding messaging node Id */
+	TOptional<FNodeEndpointId> TranslateEndpointIdToNodeId(const FGuid& EndpointId) const;
+
+private:
+
+	/** Used to unsubscribe when we're destroyed */
+	TSharedRef<IConcertServer> Server;
+	/** Used to unsubscribe when we're destroyed */
+	TSet<TWeakPtr<IConcertServerSession>> SubscribedToSessions;
+
+	/**
+	 * The client info we're caching.
+	 * 
+	 * Concert may generate multiple endpoint IDs for a single remote machine.
+	 * However, the remote node ID is always unique and is retrieved by querying the UDP backend (see GetNodeIdFromMessagingBackend).
+	 */
+	TMap<FNodeEndpointId, FConcertClientInfo> CachedClientData;
+	/**
+	 * Keeps track of past endpoints IDs that may now no longer be valid.
+	 * Every time a client joins a session, a new endpoint ID is generated for that client. It becomes impossible to
+	 * lookup old concert endpoint IDs without this mapping
+	 */
+	TMap<FGuid, FNodeEndpointId> CachedConcertEndpointToNodeEndpoints;
+
+	void OnLiveSessionCreated(bool bSuccess, const IConcertServer& InServer, TSharedRef<IConcertServerSession> InLiveSession);
+	void OnLiveSessionDestroyed(const IConcertServer& InServer, TSharedRef<IConcertServerSession> InLiveSession);
+	void OnClientInfoChanged(IConcertServerSession& Session, EConcertClientStatus ConnectionStatus, const FConcertSessionClientInfo& ClientInfo);
+	
+	void OnAdminEndpointConnectionChanged(const FConcertEndpointContext& ConcertEndpointContext, EConcertRemoteEndpointConnection ConcertRemoteEndpointConnection);
+
+	void RegisterLiveSession(const TSharedRef<IConcertServerSession>& InLiveSession);
+	void CacheClientInfo(const IConcertServerSession& Session, const FConcertSessionClientInfo& ClientInfo);
+	
+	FNodeEndpointId GetNodeIdFromMessagingBackend(const FMessageAddress& MessageAddress) const;
+};
