Visually display whether client is connected with green & gray dot widget

diff --git a/Engine/Plugins/Developer/Concert/ConcertApp/MultiUserServer/Source/MultiUserServer/Private/Widgets/Clients/Browser/Models/ClientNetworkStatisticsModel.cpp b/Engine/Plugins/Developer/Concert/ConcertApp/MultiUserServer/Source/MultiUserServer/Private/Widgets/Clients/Browser/Models/ClientNetworkStatisticsModel.cpp
--- a/Engine/Plugins/Developer/Concert/ConcertApp/MultiUserServer/Source/MultiUserServer/Private/Widgets/Clients/Browser/Models/ClientNetworkStatisticsModel.cpp
+++ b/Engine/Plugins/Developer/Concert/ConcertApp/MultiUserServer/Source/MultiUserServer/Private/Widgets/Clients/Browser/Models/ClientNetworkStatisticsModel.cpp
@@ -65,6 +65,12 @@ TOptional<FMessageTransportStatistics> UE::MultiUserServer::FClientNetworkStatis
 	return {};
 }
 
+bool UE::MultiUserServer::FClientNetworkStatisticsModel::IsOnline(const FMessageAddress& ClientAddress) const
+{
+	INetworkMessagingExtension* Statistics = Private::GetMessagingStatistics();
+	return Statistics && Statistics->GetNodeIdFromAddress(ClientAddress).IsValid();
+}
+
 void UE::MultiUserServer::FClientNetworkStatisticsModel::RegisterOnTransferUpdatedFromThread(const FMessageAddress& ClientAddress, FOnMessageTransportStatisticsUpdated StatisticUpdateCallback)
 {
 	INetworkMessagingExtension* Statistics = Private::GetMessagingStatistics();
diff --git a/Engine/Plugins/Developer/Concert/ConcertApp/MultiUserServer/Source/MultiUserServer/Private/Widgets/Clients/Browser/Models/ClientNetworkStatisticsModel.h b/Engine/Plugins/Developer/Concert/ConcertApp/MultiUserServer/Source/MultiUserServer/Private/Widgets/Clients/Browser/Models/ClientNetworkStatisticsModel.h
--- a/Engine/Plugins/Developer/Concert/ConcertApp/MultiUserServer/Source/MultiUserServer/Private/Widgets/Clients/Browser/Models/ClientNetworkStatisticsModel.h
+++ b/Engine/Plugins/Developer/Concert/ConcertApp/MultiUserServer/Source/MultiUserServer/Private/Widgets/Clients/Browser/Models/ClientNetworkStatisticsModel.h
@@ -17,6 +17,7 @@ namespace UE::MultiUserServer
 
 		//~ Begin IClientNetworkStatisticsModel Interface
 		virtual TOptional<FMessageTransportStatistics> GetLatestNetworkStatistics(const FMessageAddress& ClientAddress) const override;
+		virtual bool IsOnline(const FMessageAddress& ClientAddress) const override;
 		virtual void RegisterOnTransferUpdatedFromThread(const FMessageAddress& ClientAddress, FOnMessageTransportStatisticsUpdated StatisticsUpdatedCallback) override;
 		virtual void UnregisterOnTransferUpdatedFromThread(const FMessageAddress& ClientAddress) override;
 		//~ End IClientNetworkStatisticsModel Interface
diff --git a/Engine/Plugins/Developer/Concert/ConcertApp/MultiUserServer/Source/MultiUserServer/Private/Widgets/Clients/Browser/Models/IClientNetworkStatisticsModel.h b/Engine/Plugins/Developer/Concert/ConcertApp/MultiUserServer/Source/MultiUserServer/Private/Widgets/Clients/Browser/Models/IClientNetworkStatisticsModel.h
--- a/Engine/Plugins/Developer/Concert/ConcertApp/MultiUserServer/Source/MultiUserServer/Private/Widgets/Clients/Browser/Models/IClientNetworkStatisticsModel.h
+++ b/Engine/Plugins/Developer/Concert/ConcertApp/MultiUserServer/Source/MultiUserServer/Private/Widgets/Clients/Browser/Models/IClientNetworkStatisticsModel.h
@@ -20,6 +20,8 @@ namespace UE::MultiUserServer
 	public:
 		
 		virtual TOptional<FMessageTransportStatistics> GetLatestNetworkStatistics(const FMessageAddress& ClientAddress) const = 0;
+		/** Whether this client is currently reachable */
+		virtual bool IsOnline(const FMessageAddress& ClientAddress) const = 0;
 
 		virtual void RegisterOnTransferUpdatedFromThread(const FMessageAddress& ClientAddress, FOnMessageTransportStatisticsUpdated StatisticsUpdatedCallback) = 0;
 		virtual void UnregisterOnTransferUpdatedFromThread(const FMessageAddress& ClientAddress) = 0;
diff --git a/Engine/Plugins/Developer/Concert/ConcertApp/MultiUserServer/Source/MultiUserServer/Private/Widgets/Clients/Browser/SConcertClientBrowserItem.cpp b/Engine/Plugins/Developer/Concert/ConcertApp/MultiUserServer/Source/MultiUserServer/Private/Widgets/Clients/Browser/SConcertClientBrowserItem.cpp
--- a/Engine/Plugins/Developer/Concert/ConcertApp/MultiUserServer/Source/MultiUserServer/Private/Widgets/Clients/Browser/SConcertClientBrowserItem.cpp
+++ b/Engine/Plugins/Developer/Concert/ConcertApp/MultiUserServer/Source/MultiUserServer/Private/Widgets/Clients/Browser/SConcertClientBrowserItem.cpp
@@ -7,7 +7,9 @@
 #include "INetworkMessagingExtension.h"
 #include "SClientNetworkStats.h"
 #include "Models/IClientNetworkStatisticsModel.h"
+#include "Styling/StyleColors.h"
 #include "Widgets/SBoxPanel.h"
+#include "Widgets/Images/SImage.h"
 #include "Widgets/Layout/SBorder.h"
 #include "Widgets/Layout/SScaleBox.h"
 #include "Widgets/Text/STextBlock.h"
@@ -134,7 +136,29 @@ TSharedRef<SWidget> UE::MultiUserServer::SConcertClientBrowserItem::CreateFooter
 		.BorderImage(FConcertServerStyle::Get().GetBrush("Concert.Clients.ThumbnailFooter"))
 		[
 			SNew(SHorizontalBox)
-					
+
+			// Online / offline indicator
+			+SHorizontalBox::Slot()
+			.HAlign(HAlign_Left)
+			.Padding(2.f)
+			[
+				SNew(SImage)
+				.Image(FAppStyle::GetBrush("Icons.FilledCircle"))
+				.ColorAndOpacity_Lambda([this]()
+				{
+					return StatModel->IsOnline(Item->ClientAddress)
+						? FStyleColors::AccentGreen
+						: FStyleColors::AccentGray;
+				})
+				.ToolTipText_Lambda([this]()
+				{
+					return StatModel->IsOnline(Item->ClientAddress)
+						? LOCTEXT("ConnectionIndicator.Online", "Connected")
+						: LOCTEXT("ConnectionIndicator.Online", "Not reachable");
+				})
+			]
+
+			// IP address
 			+SHorizontalBox::Slot()
 			.HAlign(HAlign_Right)
 			.Padding(2.f)
