Fix CIS issue

diff --git a/Engine/Plugins/Developer/Concert/ConcertApp/MultiUserServer/Source/MultiUserServer/Private/Widgets/Clients/Logging/SConcertTransportLogRow.cpp b/Engine/Plugins/Developer/Concert/ConcertApp/MultiUserServer/Source/MultiUserServer/Private/Widgets/Clients/Logging/SConcertTransportLogRow.cpp
--- a/Engine/Plugins/Developer/Concert/ConcertApp/MultiUserServer/Source/MultiUserServer/Private/Widgets/Clients/Logging/SConcertTransportLogRow.cpp
+++ b/Engine/Plugins/Developer/Concert/ConcertApp/MultiUserServer/Source/MultiUserServer/Private/Widgets/Clients/Logging/SConcertTransportLogRow.cpp
@@ -7,6 +7,7 @@
 #include "SConcertTransportLog.h"
 #include "Session/Activity/PredefinedActivityColumns.h"
 
+#include "Framework/MultiBox/MultiBoxBuilder.h"
 #include "Widgets/Colors/SColorBlock.h"
 #include "Widgets/Input/SButton.h"
 #include "Widgets/Input/SComboButton.h"
