Display client transfer statistics

diff --git a/Engine/Plugins/Developer/Concert/ConcertApp/MultiUserServer/Source/MultiUserServer/Private/Widgets/Clients/Browser/Models/ClientNetworkStatisticsModel.cpp b/Engine/Plugins/Developer/Concert/ConcertApp/MultiUserServer/Source/MultiUserServer/Private/Widgets/Clients/Browser/Models/ClientNetworkStatisticsModel.cpp
--- a/Engine/Plugins/Developer/Concert/ConcertApp/MultiUserServer/Source/MultiUserServer/Private/Widgets/Clients/Browser/Models/ClientNetworkStatisticsModel.cpp
+++ b/Engine/Plugins/Developer/Concert/ConcertApp/MultiUserServer/Source/MultiUserServer/Private/Widgets/Clients/Browser/Models/ClientNetworkStatisticsModel.cpp
@@ -39,22 +39,6 @@ namespace UE::MultiUserServer::Private
 	}
 }
 
-UE::MultiUserServer::FClientNetworkStatisticsModel::FClientNetworkStatisticsModel()
-{
-	if (INetworkMessagingExtension* Statistics = Private::GetMessagingStatistics())
-	{
-		Statistics->OnTransferUpdatedFromThread().AddRaw(this, &FClientNetworkStatisticsModel::OnTransferUpdatedFromThread);
-	}
-}
-
-UE::MultiUserServer::FClientNetworkStatisticsModel::~FClientNetworkStatisticsModel()
-{
-	if (INetworkMessagingExtension* Statistics = Private::GetMessagingStatistics())
-	{
-		Statistics->OnTransferUpdatedFromThread().RemoveAll(this);
-	}
-}
-
 TOptional<FMessageTransportStatistics> UE::MultiUserServer::FClientNetworkStatisticsModel::GetLatestNetworkStatistics(const FMessageAddress& ClientAddress) const
 {
 	if (INetworkMessagingExtension* Statistics = Private::GetMessagingStatistics())
@@ -69,38 +53,4 @@ bool UE::MultiUserServer::FClientNetworkStatisticsModel::IsOnline(const FMessage
 {
 	INetworkMessagingExtension* Statistics = Private::GetMessagingStatistics();
 	return Statistics && Statistics->GetNodeIdFromAddress(ClientAddress).IsValid();
-}
-
-void UE::MultiUserServer::FClientNetworkStatisticsModel::RegisterOnTransferUpdatedFromThread(const FMessageAddress& ClientAddress, FOnMessageTransportStatisticsUpdated StatisticUpdateCallback)
-{
-	INetworkMessagingExtension* Statistics = Private::GetMessagingStatistics();
-	if (!Statistics || !ensure(IsInGameThread()))
-	{
-		return;
-	}
-	if (const FGuid NodeId = Statistics->GetNodeIdFromAddress(ClientAddress); NodeId.IsValid())
-	{
-		StatisticUpdateCallbacks.Add(NodeId, StatisticUpdateCallback);
-	}
-}
-
-void UE::MultiUserServer::FClientNetworkStatisticsModel::UnregisterOnTransferUpdatedFromThread(const FMessageAddress& ClientAddress)
-{
-	INetworkMessagingExtension* Statistics = Private::GetMessagingStatistics();
-	if (!Statistics || !ensure(IsInGameThread()))
-	{
-		return;
-	}
-	if (const FGuid NodeId = Statistics->GetNodeIdFromAddress(ClientAddress); NodeId.IsValid())
-	{
-		StatisticUpdateCallbacks.Remove(NodeId);
-	}
-}
-
-void UE::MultiUserServer::FClientNetworkStatisticsModel::OnTransferUpdatedFromThread(FTransferStatistics Stats) const
-{
-	if (const FOnMessageTransportStatisticsUpdated* Callback = StatisticUpdateCallbacks.Find(Stats.DestinationId))
-	{
-		Callback->Execute(Stats);
-	}
 }
\ No newline at end of file
diff --git a/Engine/Plugins/Developer/Concert/ConcertApp/MultiUserServer/Source/MultiUserServer/Private/Widgets/Clients/Browser/Models/ClientNetworkStatisticsModel.h b/Engine/Plugins/Developer/Concert/ConcertApp/MultiUserServer/Source/MultiUserServer/Private/Widgets/Clients/Browser/Models/ClientNetworkStatisticsModel.h
--- a/Engine/Plugins/Developer/Concert/ConcertApp/MultiUserServer/Source/MultiUserServer/Private/Widgets/Clients/Browser/Models/ClientNetworkStatisticsModel.h
+++ b/Engine/Plugins/Developer/Concert/ConcertApp/MultiUserServer/Source/MultiUserServer/Private/Widgets/Clients/Browser/Models/ClientNetworkStatisticsModel.h
@@ -12,21 +12,10 @@ namespace UE::MultiUserServer
 	{
 	public:
 
-		FClientNetworkStatisticsModel();
-		virtual ~FClientNetworkStatisticsModel() override;
-
 		//~ Begin IClientNetworkStatisticsModel Interface
 		virtual TOptional<FMessageTransportStatistics> GetLatestNetworkStatistics(const FMessageAddress& ClientAddress) const override;
 		virtual bool IsOnline(const FMessageAddress& ClientAddress) const override;
-		virtual void RegisterOnTransferUpdatedFromThread(const FMessageAddress& ClientAddress, FOnMessageTransportStatisticsUpdated StatisticsUpdatedCallback) override;
-		virtual void UnregisterOnTransferUpdatedFromThread(const FMessageAddress& ClientAddress) override;
 		//~ End IClientNetworkStatisticsModel Interface
-
-	private:
-
-		TMap<FGuid, FOnMessageTransportStatisticsUpdated> StatisticUpdateCallbacks;
-
-		void OnTransferUpdatedFromThread(FTransferStatistics Stats) const;
 	};
 }
 
diff --git a/Engine/Plugins/Developer/Concert/ConcertApp/MultiUserServer/Source/MultiUserServer/Private/Widgets/Clients/Browser/Models/ClientTransferStatisticsModel.cpp b/Engine/Plugins/Developer/Concert/ConcertApp/MultiUserServer/Source/MultiUserServer/Private/Widgets/Clients/Browser/Models/ClientTransferStatisticsModel.cpp
new file mode 100644
--- /dev/null
+++ b/Engine/Plugins/Developer/Concert/ConcertApp/MultiUserServer/Source/MultiUserServer/Private/Widgets/Clients/Browser/Models/ClientTransferStatisticsModel.cpp
@@ -0,0 +1,93 @@
+// Copyright Epic Games, Inc. All Rights Reserved.
+
+#include "ClientTransferStatisticsModel.h"
+
+#include "INetworkMessagingExtension.h"
+#include "Containers/Ticker.h"
+#include "Features/IModularFeatures.h"
+
+namespace UE::MultiUserServer::Private::ClientTransferStatisticsModel
+{
+	static INetworkMessagingExtension* GetMessagingStatistics()
+	{
+		IModularFeatures& ModularFeatures = IModularFeatures::Get();
+		if (IsInGameThread())
+		{
+			if (ModularFeatures.IsModularFeatureAvailable(INetworkMessagingExtension::ModularFeatureName))
+			{
+				return &ModularFeatures.GetModularFeature<INetworkMessagingExtension>(INetworkMessagingExtension::ModularFeatureName);
+			}
+		}
+		else
+		{
+			ModularFeatures.LockModularFeatureList();
+			ON_SCOPE_EXIT
+			{
+				ModularFeatures.UnlockModularFeatureList();
+			};
+			
+			if (ModularFeatures.IsModularFeatureAvailable(INetworkMessagingExtension::ModularFeatureName))
+			{
+				return &ModularFeatures.GetModularFeature<INetworkMessagingExtension>(INetworkMessagingExtension::ModularFeatureName);
+			}
+		}
+		
+		ensureMsgf(false, TEXT("Feature %s is unavailable"), *INetworkMessagingExtension::ModularFeatureName.ToString());
+		return nullptr;
+	}
+}
+
+UE::MultiUserServer::FClientTransferStatisticsModel::FClientTransferStatisticsModel(const FMessageAddress& MessageAddress)
+	: MessageAddress(MessageAddress)
+{
+	if (INetworkMessagingExtension* Statistics = Private::ClientTransferStatisticsModel::GetMessagingStatistics())
+	{
+		Statistics->OnTransferUpdatedFromThread().AddRaw(this, &FClientTransferStatisticsModel::OnTransferUpdatedFromThread);
+		TickHandle = FTSTicker::GetCoreTicker().AddTicker(FTickerDelegate::CreateRaw(this, &FClientTransferStatisticsModel::Tick));
+	}
+}
+
+UE::MultiUserServer::FClientTransferStatisticsModel::~FClientTransferStatisticsModel()
+{
+	if (INetworkMessagingExtension* Statistics = Private::ClientTransferStatisticsModel::GetMessagingStatistics())
+	{
+		Statistics->OnTransferUpdatedFromThread().RemoveAll(this);
+	}
+	
+	FTSTicker::GetCoreTicker().RemoveTicker(TickHandle);
+}
+
+void UE::MultiUserServer::FClientTransferStatisticsModel::OnTransferUpdatedFromThread(FTransferStatistics TransferStatistics)
+{
+	AsyncStatQueue.Enqueue(TransferStatistics);
+}
+
+bool UE::MultiUserServer::FClientTransferStatisticsModel::Tick(float DeltaTime)
+{
+	bool bAnyElements = false;
+	while (const TOptional<FTransferStatistics> TransferStatistics = AsyncStatQueue.Dequeue())
+	{
+		bAnyElements = true;
+		
+		const TSharedPtr<FTransferStatistics> NewValue = MakeShared<FTransferStatistics>(*TransferStatistics);
+		auto Pos = Algo::UpperBound(Stats, NewValue, [](const TSharedPtr<FTransferStatistics>& Value, const TSharedPtr<FTransferStatistics>& Check)
+		{
+			return Value->MessageId > Check->MessageId;
+		});
+		if (Stats.Num() > 0 && Pos < Stats.Num() && Stats[Pos] && Stats[Pos]->MessageId == NewValue->MessageId)
+		{
+			Stats[Pos] = NewValue;
+		}
+		else
+		{
+			Stats.Insert(NewValue, Pos);
+		}
+	}
+
+	if (bAnyElements)
+	{
+		OnUpdatedDelegate.Broadcast();
+	}
+
+	return true;
+}
diff --git a/Engine/Plugins/Developer/Concert/ConcertApp/MultiUserServer/Source/MultiUserServer/Private/Widgets/Clients/Browser/Models/ClientTransferStatisticsModel.h b/Engine/Plugins/Developer/Concert/ConcertApp/MultiUserServer/Source/MultiUserServer/Private/Widgets/Clients/Browser/Models/ClientTransferStatisticsModel.h
new file mode 100644
--- /dev/null
+++ b/Engine/Plugins/Developer/Concert/ConcertApp/MultiUserServer/Source/MultiUserServer/Private/Widgets/Clients/Browser/Models/ClientTransferStatisticsModel.h
@@ -0,0 +1,36 @@
+// Copyright Epic Games, Inc. All Rights Reserved.
+
+#pragma once
+
+#include "CoreMinimal.h"
+#include "IClientTransferStatisticsModel.h"
+#include "Containers/SpscQueue.h"
+#include "Containers/Ticker.h"
+
+namespace UE::MultiUserServer
+{
+	class FClientTransferStatisticsModel : public IClientTransferStatisticsModel
+	{
+	public:
+		
+		FClientTransferStatisticsModel(const FMessageAddress& MessageAddress);
+		virtual ~FClientTransferStatisticsModel() override;
+
+		virtual const TArray<TSharedPtr<FTransferStatistics>>& GetSortedTransferStatistics() const override { return Stats; }
+		virtual FOnTransferStatisticsUpdated& OnTransferStatisticsUpdated() override { return OnUpdatedDelegate; }
+
+	private:
+
+		const FMessageAddress MessageAddress;
+
+		TSpscQueue<FTransferStatistics> AsyncStatQueue;
+		TArray<TSharedPtr<FTransferStatistics>> Stats;
+		
+		FOnTransferStatisticsUpdated OnUpdatedDelegate;
+		FTSTicker::FDelegateHandle TickHandle;
+		
+		void OnTransferUpdatedFromThread(FTransferStatistics TransferStatistics);
+		bool Tick(float DeltaTime);
+	};
+}
+
diff --git a/Engine/Plugins/Developer/Concert/ConcertApp/MultiUserServer/Source/MultiUserServer/Private/Widgets/Clients/Browser/Models/IClientNetworkStatisticsModel.h b/Engine/Plugins/Developer/Concert/ConcertApp/MultiUserServer/Source/MultiUserServer/Private/Widgets/Clients/Browser/Models/IClientNetworkStatisticsModel.h
--- a/Engine/Plugins/Developer/Concert/ConcertApp/MultiUserServer/Source/MultiUserServer/Private/Widgets/Clients/Browser/Models/IClientNetworkStatisticsModel.h
+++ b/Engine/Plugins/Developer/Concert/ConcertApp/MultiUserServer/Source/MultiUserServer/Private/Widgets/Clients/Browser/Models/IClientNetworkStatisticsModel.h
@@ -12,8 +12,6 @@ struct FTransferStatistics;
 
 namespace UE::MultiUserServer
 {
-	DECLARE_DELEGATE_OneParam(FOnMessageTransportStatisticsUpdated, const FTransferStatistics&);
-	
 	/** Decouples the UI from the server functions. */
 	class IClientNetworkStatisticsModel
 	{
@@ -23,9 +21,6 @@ namespace UE::MultiUserServer
 		/** Whether this client is currently reachable */
 		virtual bool IsOnline(const FMessageAddress& ClientAddress) const = 0;
 
-		virtual void RegisterOnTransferUpdatedFromThread(const FMessageAddress& ClientAddress, FOnMessageTransportStatisticsUpdated StatisticsUpdatedCallback) = 0;
-		virtual void UnregisterOnTransferUpdatedFromThread(const FMessageAddress& ClientAddress) = 0;
-
 		virtual ~IClientNetworkStatisticsModel() = default;
 	};
 }
diff --git a/Engine/Plugins/Developer/Concert/ConcertApp/MultiUserServer/Source/MultiUserServer/Private/Widgets/Clients/Browser/Models/IClientTransferStatisticsModel.h b/Engine/Plugins/Developer/Concert/ConcertApp/MultiUserServer/Source/MultiUserServer/Private/Widgets/Clients/Browser/Models/IClientTransferStatisticsModel.h
new file mode 100644
--- /dev/null
+++ b/Engine/Plugins/Developer/Concert/ConcertApp/MultiUserServer/Source/MultiUserServer/Private/Widgets/Clients/Browser/Models/IClientTransferStatisticsModel.h
@@ -0,0 +1,30 @@
+// Copyright Epic Games, Inc. All Rights Reserved.
+
+#pragma once
+
+#include "CoreMinimal.h"
+#include "IMessageContext.h"
+
+struct FMessageAddress;
+struct FTransferStatistics;
+
+namespace UE::MultiUserServer
+{
+	DECLARE_MULTICAST_DELEGATE(FOnTransferStatisticsUpdated);
+	
+	/** Keeps track of a single client's transferstatistics */
+	class IClientTransferStatisticsModel
+	{
+	public:
+
+		/** Gets the transfer statistics sorted descending by message ID */
+		virtual const TArray<TSharedPtr<FTransferStatistics>>& GetSortedTransferStatistics() const = 0;
+
+		/** Called when the transfer statistics change */
+		virtual FOnTransferStatisticsUpdated& OnTransferStatisticsUpdated() = 0;
+
+		virtual ~IClientTransferStatisticsModel() = default;
+	};
+}
+
+
diff --git a/Engine/Plugins/Developer/Concert/ConcertApp/MultiUserServer/Source/MultiUserServer/Private/Widgets/Clients/Browser/SClientTransferStatTable.cpp b/Engine/Plugins/Developer/Concert/ConcertApp/MultiUserServer/Source/MultiUserServer/Private/Widgets/Clients/Browser/SClientTransferStatTable.cpp
new file mode 100644
--- /dev/null
+++ b/Engine/Plugins/Developer/Concert/ConcertApp/MultiUserServer/Source/MultiUserServer/Private/Widgets/Clients/Browser/SClientTransferStatTable.cpp
@@ -0,0 +1,143 @@
+// Copyright Epic Games, Inc. All Rights Reserved.
+
+#include "SClientTransferStatTable.h"
+
+#include "INetworkMessagingExtension.h"
+#include "Models/IClientTransferStatisticsModel.h"
+#include "Widgets/SBoxPanel.h"
+#include "Widgets/SNullWidget.h"
+#include "Widgets/Text/STextBlock.h"
+#include "Widgets/Views/SHeaderRow.h"
+#include "Widgets/Views/SListView.h"
+
+#define LOCTEXT_NAMESPACE "UnrealMultiUserUI.SClientTransferStatTable"
+
+namespace UE::MultiUserServer
+{
+	namespace Private
+	{
+		const FName HeaderIdName_MessageId    = TEXT("MessageId");
+		const FName HeaderIdName_SentSegments = TEXT("Sent");
+		const FName HeaderIdName_AckSegments  = TEXT("SegmentsAck");
+		const FName HeaderIdName_TotalSize    = TEXT("Size");
+		const FName HeaderIdName_DataRate     = TEXT("DataRate");
+	}
+
+	class SClientTransferStatTableRow : public SMultiColumnTableRow<TSharedRef<FTransferStatistics>>
+	{
+		using Super = SMultiColumnTableRow<TSharedRef<FTransferStatistics>>;
+	public:
+		
+		SLATE_BEGIN_ARGS(SClientTransferStatTableRow) { }
+			SLATE_ARGUMENT(TSharedPtr<FTransferStatistics>, Stats)
+		SLATE_END_ARGS()
+		
+		void Construct(const FArguments& InArgs, const TSharedRef<STableViewBase>& InOwerTableView)
+		{
+			Stats = InArgs._Stats;
+			Super::Construct(Super::FArguments(), InOwerTableView);
+		}
+		
+	private:
+
+		TSharedPtr<FTransferStatistics> Stats;
+
+		virtual TSharedRef<SWidget> GenerateWidgetForColumn(const FName& ColumnName) override
+		{
+			using namespace Private;
+			if (HeaderIdName_MessageId == ColumnName)
+			{
+				return SNew(SHorizontalBox)
+					+ SHorizontalBox::Slot()
+					.HAlign(HAlign_Left)
+					.VAlign(VAlign_Center)
+					.AutoWidth()
+					[
+						SNew(STextBlock)
+						.Text_Lambda([this](){ return FText::AsNumber(Stats->MessageId); })
+					];
+			}
+			if (HeaderIdName_SentSegments == ColumnName)
+			{
+				return SNew(SHorizontalBox)
+					+ SHorizontalBox::Slot()
+					.HAlign(HAlign_Left)
+					.VAlign(VAlign_Center)
+					.AutoWidth()
+					[
+						SNew(STextBlock)
+						.Text_Lambda([this](){ return FText::AsNumber(Stats->BytesSent); })
+					];
+			}
+			if (HeaderIdName_AckSegments == ColumnName)
+			{
+				return SNew(SHorizontalBox)
+					+ SHorizontalBox::Slot()
+					.HAlign(HAlign_Left)
+					.VAlign(VAlign_Center)
+					.AutoWidth()
+					[
+						SNew(STextBlock)
+						.Text_Lambda([this](){ return FText::AsNumber(Stats->BytesAcknowledged); })
+					];
+			}
+			if (HeaderIdName_TotalSize == ColumnName)
+			{
+				return SNew(SHorizontalBox)
+					+ SHorizontalBox::Slot()
+					.HAlign(HAlign_Left)
+					.VAlign(VAlign_Center)
+					.AutoWidth()
+					[
+						SNew(STextBlock)
+						.Text_Lambda([this](){ return FText::AsNumber(Stats->BytesToSend); })
+					];
+			}
+			return SNullWidget::NullWidget;
+		}
+	};
+	
+	void SClientTransferStatTable::Construct(const FArguments& InArgs, TSharedRef<IClientTransferStatisticsModel> InStatsModel)
+	{
+		StatsModel = InStatsModel;
+
+		StatsModel->OnTransferStatisticsUpdated().AddSP(this, &SClientTransferStatTable::OnTransferStatisticsUpdated);
+
+		ChildSlot
+		[
+			SAssignNew(SegmenterListView, SListView<TSharedPtr<FTransferStatistics>>)
+			.ListItemsSource(&InStatsModel->GetSortedTransferStatistics())
+			.OnGenerateRow(this, &SClientTransferStatTable::OnGenerateActivityRowWidget)
+			.SelectionMode(ESelectionMode::None)
+			.HeaderRow
+			(
+				SNew(SHeaderRow)
+				+ SHeaderRow::Column(Private::HeaderIdName_MessageId)
+				.FillWidth(2.f)
+				.DefaultLabel(LOCTEXT("HeaderName_MessageId", "Id"))
+				+ SHeaderRow::Column(Private::HeaderIdName_SentSegments)
+				.FillWidth(2.f)
+				.DefaultLabel(LOCTEXT("HeaderName_SentSegments", "Sent"))
+				+ SHeaderRow::Column(Private::HeaderIdName_AckSegments)
+				.FillWidth(2.f)
+				.DefaultLabel(LOCTEXT("HeaderName_AckSegments", "Ack"))
+				+ SHeaderRow::Column(Private::HeaderIdName_TotalSize)
+				.FillWidth(2.f)
+				.DefaultLabel(LOCTEXT("HeaderName_TotalSize", "Size"))
+			)
+		];
+	}
+
+	TSharedRef<ITableRow> SClientTransferStatTable::OnGenerateActivityRowWidget(TSharedPtr<FTransferStatistics> InStats, const TSharedRef<STableViewBase>& OwnerTable) const
+	{
+		return SNew(SClientTransferStatTableRow, OwnerTable)
+			.Stats(InStats);;
+	}
+
+	void SClientTransferStatTable::OnTransferStatisticsUpdated() const
+	{
+		SegmenterListView->RequestListRefresh();
+	}
+}
+
+#undef LOCTEXT_NAMESPACE
\ No newline at end of file
diff --git a/Engine/Plugins/Developer/Concert/ConcertApp/MultiUserServer/Source/MultiUserServer/Private/Widgets/Clients/Browser/SClientTransferStatTable.h b/Engine/Plugins/Developer/Concert/ConcertApp/MultiUserServer/Source/MultiUserServer/Private/Widgets/Clients/Browser/SClientTransferStatTable.h
new file mode 100644
--- /dev/null
+++ b/Engine/Plugins/Developer/Concert/ConcertApp/MultiUserServer/Source/MultiUserServer/Private/Widgets/Clients/Browser/SClientTransferStatTable.h
@@ -0,0 +1,42 @@
+// Copyright Epic Games, Inc. All Rights Reserved.
+
+#pragma once
+
+#include "CoreMinimal.h"
+#include "Widgets/SCompoundWidget.h"
+
+class ITableRow;
+template<typename T> class SListView;
+class STableViewBase;
+struct FTransferStatistics;
+
+namespace UE::MultiUserServer
+{
+	class IClientTransferStatisticsModel;
+
+	/**
+	 * Displays FTransferStatistics as they are updated by INetworkMessagingExtension.
+	 */
+	class SClientTransferStatTable : public SCompoundWidget
+	{
+	public:
+
+		SLATE_BEGIN_ARGS(SClientTransferStatTable)
+		{}
+		SLATE_END_ARGS()
+
+		void Construct(const FArguments& InArgs, TSharedRef<IClientTransferStatisticsModel> InStatsModel);
+
+	private:
+
+		/** Tells us when the transfer stats have changed */
+		TSharedPtr<IClientTransferStatisticsModel> StatsModel;
+		/** Displays the transfer stats */
+		TSharedPtr<SListView<TSharedPtr<FTransferStatistics>>> SegmenterListView;
+
+		TSharedRef<ITableRow> OnGenerateActivityRowWidget(TSharedPtr<FTransferStatistics> InStats, const TSharedRef<STableViewBase>& OwnerTable) const;
+		void OnTransferStatisticsUpdated() const;
+	};
+}
+
+
diff --git a/Engine/Plugins/Developer/Concert/ConcertApp/MultiUserServer/Source/MultiUserServer/Private/Widgets/Clients/Browser/SConcertClientBrowser.cpp b/Engine/Plugins/Developer/Concert/ConcertApp/MultiUserServer/Source/MultiUserServer/Private/Widgets/Clients/Browser/SConcertClientBrowser.cpp
--- a/Engine/Plugins/Developer/Concert/ConcertApp/MultiUserServer/Source/MultiUserServer/Private/Widgets/Clients/Browser/SConcertClientBrowser.cpp
+++ b/Engine/Plugins/Developer/Concert/ConcertApp/MultiUserServer/Source/MultiUserServer/Private/Widgets/Clients/Browser/SConcertClientBrowser.cpp
@@ -9,12 +9,14 @@
 #include "Models/IClientNetworkStatisticsModel.h"
 #include "SConcertClientBrowserItem.h"
 #include "Algo/Accumulate.h"
+#include "Algo/AllOf.h"
 
 #include "Algo/AnyOf.h"
 #include "Algo/Count.h"
 #include "Dialog/SMessageDialog.h"
 #include "Framework/MultiBox/MultiBoxBuilder.h"
 #include "Widgets/Input/SCheckBox.h"
+#include "Widgets/Input/SComboButton.h"
 #include "Widgets/Input/SSearchBox.h"
 #include "Widgets/SBoxPanel.h"
 #include "Widgets/Text/STextBlock.h"
@@ -189,7 +191,7 @@ TSharedRef<SWidget> UE::MultiUserServer::SConcertClientBrowser::CreateKeepDiscon
 
 TSharedRef<SWidget> UE::MultiUserServer::SConcertClientBrowser::CreateTileView()
 {
-	constexpr int32 Height = 200;
+	constexpr int32 Height = 270;
 	constexpr int32 Width = 270;
 	return SAssignNew(TileView, STileView<TSharedPtr<FClientBrowserItem>>)
 		.SelectionMode(ESelectionMode::Multi)
@@ -306,7 +308,6 @@ TSharedRef<SWidget> UE::MultiUserServer::SConcertClientBrowser::SConcertClientBr
 			EUserInterfaceActionType::Check
 		);
 	}
-
 	return MenuBuilder.MakeWidget();
 }
 
@@ -374,9 +375,53 @@ TSharedPtr<SWidget> UE::MultiUserServer::SConcertClientBrowser::OnGetContextMenu
 		EUserInterfaceActionType::Button
 		);
 	
+	MenuBuilder.AddSeparator();
+	AddDisplayModeEntry(MenuBuilder, EClientDisplayMode::NetworkGraph,
+		LOCTEXT("DisplayMode.NetworkGraph.Title", "Network graph"),
+		LOCTEXT("DisplayMode.NetworkGraph.Tooltip", "Show the up and down stream network traffic on a graph")
+		);
+	AddDisplayModeEntry(MenuBuilder, EClientDisplayMode::SegementTable,
+		LOCTEXT("DisplayMode.SegementTable.Title", "Segment table"),
+		LOCTEXT("DisplayMode.SegementTable.Tooltip", "A table displaying the messaging protocol's segments MessageId, Sent, Acked and Size in realtime.")
+		);
 	return MenuBuilder.MakeWidget();
 }
 
+void UE::MultiUserServer::SConcertClientBrowser::AddDisplayModeEntry(FMenuBuilder& MenuBuilder, EClientDisplayMode DisplayMode, FText Title, FText Tooltip) const
+{
+	MenuBuilder.AddMenuEntry(
+		Title,
+		Tooltip,
+		FSlateIcon(),
+		FUIAction(
+			FExecuteAction::CreateLambda([this, DisplayMode]()
+			{
+				TArray<TSharedPtr<FClientBrowserItem>> Items = TileView->GetSelectedItems();
+				for (const TSharedPtr<FClientBrowserItem>& Item : Items)
+				{
+					if (const TSharedPtr<SConcertClientBrowserItem>* ItemWidget = ClientWidgets.Find(Item->MessageNodeId))
+					{
+						ItemWidget->Get()->SetDisplayMode(DisplayMode);
+					}
+				}
+			}),
+			FCanExecuteAction::CreateLambda([] { return true; }),
+			FIsActionChecked::CreateLambda([this, DisplayMode]()
+			{
+				return Algo::AllOf(TileView->GetSelectedItems(), [this, DisplayMode](const TSharedPtr<FClientBrowserItem>& Item)
+				{
+					if (const TSharedPtr<SConcertClientBrowserItem>* ItemWidget = ClientWidgets.Find(Item->MessageNodeId))
+					{
+						return ItemWidget->Get()->GetDisplayMode() == DisplayMode;
+					}
+					return false;
+				});
+			})),
+		NAME_None,
+		EUserInterfaceActionType::Check
+	);
+}
+
 void UE::MultiUserServer::SConcertClientBrowser::OnListMouseButtonDoubleClick(TSharedPtr<FClientBrowserItem> ClientItem)
 {
 	OnClientDoubleClicked.ExecuteIfBound(ClientItem->MessageNodeId);
@@ -407,6 +452,18 @@ void UE::MultiUserServer::SConcertClientBrowser::UpdateTileViewFromAllowedSessio
 		}
 	}
 	TileView->RequestListRefresh();
+
+	for (auto CachedClientWidgetIt = ClientWidgets.CreateIterator(); CachedClientWidgetIt; ++CachedClientWidgetIt)
+	{
+		const bool bShouldExist = BrowserModel->GetItems().ContainsByPredicate([MessageNodeId = CachedClientWidgetIt->Key](const TSharedPtr<FClientBrowserItem>& Item)
+		{
+			return Item->MessageNodeId == MessageNodeId;
+		});
+		if (!bShouldExist)
+		{
+			CachedClientWidgetIt.RemoveCurrent();
+		}
+	}
 }
 
 bool UE::MultiUserServer::SConcertClientBrowser::PassesFilter(const TSharedPtr<FClientBrowserItem>& Client) const
diff --git a/Engine/Plugins/Developer/Concert/ConcertApp/MultiUserServer/Source/MultiUserServer/Private/Widgets/Clients/Browser/SConcertClientBrowser.h b/Engine/Plugins/Developer/Concert/ConcertApp/MultiUserServer/Source/MultiUserServer/Private/Widgets/Clients/Browser/SConcertClientBrowser.h
--- a/Engine/Plugins/Developer/Concert/ConcertApp/MultiUserServer/Source/MultiUserServer/Private/Widgets/Clients/Browser/SConcertClientBrowser.h
+++ b/Engine/Plugins/Developer/Concert/ConcertApp/MultiUserServer/Source/MultiUserServer/Private/Widgets/Clients/Browser/SConcertClientBrowser.h
@@ -10,6 +10,7 @@
 #include "Widgets/SCompoundWidget.h"
 
 class FConcertClientsTabController;
+class FMenuBuilder;
 class STableViewBase;
 class ITableRow;
 template<typename T> class STileView;
@@ -21,6 +22,7 @@ namespace UE::MultiUserServer
 {
 	class IClientNetworkStatisticsModel;
 	class IClientBrowserModel;
+	enum class EClientDisplayMode : uint8;
 	struct FClientBrowserItem;
 	
 	/**
@@ -92,6 +94,7 @@ namespace UE::MultiUserServer
 		// TileView events
 		TSharedRef<ITableRow> MakeTileViewWidget(TSharedPtr<FClientBrowserItem> ClientItem, const TSharedRef<STableViewBase>& OwnerTable);
 		TSharedPtr<SWidget> OnGetContextMenuContent();
+		void AddDisplayModeEntry(FMenuBuilder& MenuBuilder, EClientDisplayMode DisplayMode, FText Title, FText Tooltip) const;
 		void OnListMouseButtonDoubleClick(TSharedPtr<FClientBrowserItem> ClientItem);
 
 		// Filtering
diff --git a/Engine/Plugins/Developer/Concert/ConcertApp/MultiUserServer/Source/MultiUserServer/Private/Widgets/Clients/Browser/SConcertClientBrowserItem.cpp b/Engine/Plugins/Developer/Concert/ConcertApp/MultiUserServer/Source/MultiUserServer/Private/Widgets/Clients/Browser/SConcertClientBrowserItem.cpp
--- a/Engine/Plugins/Developer/Concert/ConcertApp/MultiUserServer/Source/MultiUserServer/Private/Widgets/Clients/Browser/SConcertClientBrowserItem.cpp
+++ b/Engine/Plugins/Developer/Concert/ConcertApp/MultiUserServer/Source/MultiUserServer/Private/Widgets/Clients/Browser/SConcertClientBrowserItem.cpp
@@ -6,13 +6,18 @@
 #include "ConcertServerStyle.h"
 #include "INetworkMessagingExtension.h"
 #include "SClientNetworkStats.h"
+#include "SClientTransferStatTable.h"
+#include "Models/ClientTransferStatisticsModel.h"
 #include "Models/IClientNetworkStatisticsModel.h"
 #include "Styling/StyleColors.h"
-#include "Widgets/SBoxPanel.h"
 #include "Widgets/Images/SImage.h"
+#include "Widgets/LayerManager/STooltipPresenter.h"
 #include "Widgets/Layout/SBorder.h"
 #include "Widgets/Layout/SScaleBox.h"
+#include "Widgets/SOverlay.h"
 #include "Widgets/Text/STextBlock.h"
+#include "Widgets/SBoxPanel.h"
+#include "Widgets/SNullWidget.h"
 
 #define LOCTEXT_NAMESPACE "UnrealMultiUserUI.SConcertClientBrowserItem"
 
@@ -20,6 +25,7 @@ void UE::MultiUserServer::SConcertClientBrowserItem::Construct(const FArguments&
 {
 	Item = MoveTemp(InClientItem);
 	StatModel = MoveTemp(InStatModel);
+	TransferStatsModel = MakeShared<FClientTransferStatisticsModel>(InClientItem->ClientAddress);
 	HighlightText = InArgs._HighlightText;
 
 	ChildSlot
@@ -43,11 +49,18 @@ void UE::MultiUserServer::SConcertClientBrowserItem::Construct(const FArguments&
 					SNew(SVerticalBox)
 
 					+SVerticalBox::Slot()
-					.FillHeight(1.f)
+					.AutoHeight()
 					.VAlign(VAlign_Top)
 					[
 						CreateHeader()	
 					]
+					
+					+SVerticalBox::Slot()
+					.FillHeight(1.f)
+					.Padding(0.f, 5.f, 0.f, 0.f)
+					[
+						CreateContentArea()
+					]
 
 					+SVerticalBox::Slot()
 					.AutoHeight()
@@ -119,6 +132,21 @@ TSharedRef<SWidget> UE::MultiUserServer::SConcertClientBrowserItem::CreateHeader
 		.ColorAndOpacity(FColor::White);
 }
 
+TSharedRef<SWidget> UE::MultiUserServer::SConcertClientBrowserItem::CreateContentArea()
+{
+	return SNew(SOverlay)
+		+SOverlay::Slot()
+		[
+			// TODO: Graph
+			SNullWidget::NullWidget
+		]
+		+SOverlay::Slot()
+		[
+			SNew(SClientTransferStatTable, TransferStatsModel.ToSharedRef())
+			.Visibility_Lambda([this](){ return GetDisplayMode() == EClientDisplayMode::SegementTable ? EVisibility::Visible : EVisibility::Collapsed; })
+		];
+}
+
 TSharedRef<SWidget> UE::MultiUserServer::SConcertClientBrowserItem::CreateStats()
 {
 	return SNew(SScaleBox)
@@ -154,7 +182,7 @@ TSharedRef<SWidget> UE::MultiUserServer::SConcertClientBrowserItem::CreateFooter
 				{
 					return StatModel->IsOnline(Item->ClientAddress)
 						? LOCTEXT("ConnectionIndicator.Online", "Connected")
-						: LOCTEXT("ConnectionIndicator.Online", "Not reachable");
+						: LOCTEXT("ConnectionIndicator.Offline", "Not reachable");
 				})
 			]
 
diff --git a/Engine/Plugins/Developer/Concert/ConcertApp/MultiUserServer/Source/MultiUserServer/Private/Widgets/Clients/Browser/SConcertClientBrowserItem.h b/Engine/Plugins/Developer/Concert/ConcertApp/MultiUserServer/Source/MultiUserServer/Private/Widgets/Clients/Browser/SConcertClientBrowserItem.h
--- a/Engine/Plugins/Developer/Concert/ConcertApp/MultiUserServer/Source/MultiUserServer/Private/Widgets/Clients/Browser/SConcertClientBrowserItem.h
+++ b/Engine/Plugins/Developer/Concert/ConcertApp/MultiUserServer/Source/MultiUserServer/Private/Widgets/Clients/Browser/SConcertClientBrowserItem.h
@@ -4,9 +4,11 @@
 
 #include "CoreMinimal.h"
 #include "IMessageContext.h"
+#include "Models/IClientTransferStatisticsModel.h"
 #include "Widgets/DeclarativeSyntaxSupport.h"
 #include "Widgets/SCompoundWidget.h"
 
+class FMenuBuilder;
 class STextBlock;
 
 namespace UE::MultiUserServer
@@ -15,6 +17,14 @@ namespace UE::MultiUserServer
 	class IClientNetworkStatisticsModel;
 	struct FClientBrowserItem;
 
+	enum class EClientDisplayMode : uint8
+	{
+		/** Displays the sent and read packets */
+		NetworkGraph = 0,
+		/** Displays a table showing MessageId, Sent, Acked, and Size updated in realtime. */
+		SegementTable
+	};
+	
 	class SConcertClientBrowserItem : public SCompoundWidget
 	{
 	public:
@@ -30,11 +40,15 @@ namespace UE::MultiUserServer
 		/** Appends the statistics to the search terms */
 		void AppendSearchTerms(TArray<FString>& SearchTerms) const;
 
+		void SetDisplayMode(EClientDisplayMode Value) { DisplayMode = Value; }
+		EClientDisplayMode GetDisplayMode() const { return DisplayMode;}
+
 	private:
 
 		/** What we're displaying */
 		TSharedPtr<FClientBrowserItem> Item;
 		TSharedPtr<IClientNetworkStatisticsModel> StatModel;
+		TSharedPtr<IClientTransferStatisticsModel> TransferStatsModel;
 
 		/** The text to highlight */
 		TSharedPtr<FText> HighlightText;
@@ -43,7 +57,10 @@ namespace UE::MultiUserServer
 		TSharedPtr<SClientNetworkStats> NetworkStats;
 		TSharedPtr<STextBlock> ClientIP4;
 
+		EClientDisplayMode DisplayMode = EClientDisplayMode::NetworkGraph;
+
 		TSharedRef<SWidget> CreateHeader();
+		TSharedRef<SWidget> CreateContentArea();
 		TSharedRef<SWidget> CreateStats();
 		TSharedRef<SWidget> CreateFooter();
 
diff --git a/Engine/Plugins/Developer/Concert/ConcertApp/MultiUserServer/Source/MultiUserServer/Private/Window/ConcertServerWindowController.cpp b/Engine/Plugins/Developer/Concert/ConcertApp/MultiUserServer/Source/MultiUserServer/Private/Window/ConcertServerWindowController.cpp
--- a/Engine/Plugins/Developer/Concert/ConcertApp/MultiUserServer/Source/MultiUserServer/Private/Window/ConcertServerWindowController.cpp
+++ b/Engine/Plugins/Developer/Concert/ConcertApp/MultiUserServer/Source/MultiUserServer/Private/Window/ConcertServerWindowController.cpp
@@ -46,7 +46,7 @@ TSharedRef<SWindow> FConcertServerWindowController::CreateWindow()
 	const float DPIScaleFactor = FPlatformApplicationMisc::GetDPIScaleFactorAtPoint(DisplayMetrics.PrimaryDisplayWorkAreaRect.Left, DisplayMetrics.PrimaryDisplayWorkAreaRect.Top);
 
 	const bool bEmbedTitleAreaContent = false;
-	const FVector2D ClientSize(960.0f * DPIScaleFactor, 640.0f * DPIScaleFactor);
+	const FVector2D ClientSize(960.0f * DPIScaleFactor, 750.0f * DPIScaleFactor);
 	TSharedRef<SWindow> RootWindowRef = SNew(SWindow)
 		.Title(LOCTEXT("WindowTitle", "Unreal Multi User Server"))
 		.CreateTitleBar(!bEmbedTitleAreaContent)
