Remove unused style resource

diff --git a/Engine/Plugins/Developer/Concert/ConcertApp/MultiUserServer/Source/MultiUserServer/Private/ConcertServerStyle.cpp b/Engine/Plugins/Developer/Concert/ConcertApp/MultiUserServer/Source/MultiUserServer/Private/ConcertServerStyle.cpp
--- a/Engine/Plugins/Developer/Concert/ConcertApp/MultiUserServer/Source/MultiUserServer/Private/ConcertServerStyle.cpp
+++ b/Engine/Plugins/Developer/Concert/ConcertApp/MultiUserServer/Source/MultiUserServer/Private/ConcertServerStyle.cpp
@@ -71,7 +71,6 @@ TSharedRef<FSlateStyleSet> FConcertServerStyle::Create()
 	StyleSet->Set("Concert.Ack.Ack", new IMAGE_PLUGIN_BRUSH("Icons/Ack_Ack_16x", Icon16x16));
 	StyleSet->Set("Concert.Ack.Success", new IMAGE_PLUGIN_BRUSH("Icons/Ack_Success_16x", Icon16x16));
 	StyleSet->Set("Concert.Ack.Failure", new IMAGE_PLUGIN_BRUSH("Icons/Ack_Fail_16x", Icon16x16));
-	StyleSet->Set("Concert.Ack.NotNeeded", new IMAGE_PLUGIN_BRUSH("Icons/Ack_NotNeeded_16x", Icon16x16));
 
 	// Clients tab
 	StyleSet->Set("Concert.Clients.DropShadow", new IMAGE_PLUGIN_BOX_BRUSH("ClientThumbnailDropShadow", FMargin(4.0f / 64.0f)));
