Show all log messages by default

diff --git a/Engine/Plugins/Developer/Concert/ConcertApp/MultiUserServer/Source/MultiUserServer/Private/Widgets/StatusBar/SConcertStatusBar.cpp b/Engine/Plugins/Developer/Concert/ConcertApp/MultiUserServer/Source/MultiUserServer/Private/Widgets/StatusBar/SConcertStatusBar.cpp
--- a/Engine/Plugins/Developer/Concert/ConcertApp/MultiUserServer/Source/MultiUserServer/Private/Widgets/StatusBar/SConcertStatusBar.cpp
+++ b/Engine/Plugins/Developer/Concert/ConcertApp/MultiUserServer/Source/MultiUserServer/Private/Widgets/StatusBar/SConcertStatusBar.cpp
@@ -30,20 +30,6 @@ namespace UE::ConcertServerUI::Private
 				Params.SettingsMenuCreationFlags = EOutputLogSettingsMenuFlags::SkipClearOnPie
 					| EOutputLogSettingsMenuFlags::SkipOpenSourceButton
 					| EOutputLogSettingsMenuFlags::SkipEnableWordWrapping; // Checkbox relies on saving an editor config file and does not work correctly
-				Params.AllowAsInitialLogCategory = FAllowLogCategoryCallback::CreateLambda([](const FName CategoryName)
-				{
-					// Hide these categories for improved UX
-					static TArray<FString> DeselectedByDefault = { "Slate", "WindowsTextInput" };
-					const FString CategoryNameAsString = CategoryName.ToString();
-					for (const FString& DeselectedCategory : DeselectedByDefault)
-					{
-						if (CategoryNameAsString.Contains(DeselectedCategory))
-						{
-							return false;
-						}
-					}
-					return true;
-				});
 				StatusBarOutputLog = FOutputLogModule::Get().MakeOutputLogWidget(Params);
 			}
 
