Create UI for muting activities

diff --git a/Engine/Plugins/Developer/Concert/ConcertApp/MultiUserServer/Source/MultiUserServer/Private/ConcertServerStyle.cpp b/Engine/Plugins/Developer/Concert/ConcertApp/MultiUserServer/Source/MultiUserServer/Private/ConcertServerStyle.cpp
--- a/Engine/Plugins/Developer/Concert/ConcertApp/MultiUserServer/Source/MultiUserServer/Private/ConcertServerStyle.cpp
+++ b/Engine/Plugins/Developer/Concert/ConcertApp/MultiUserServer/Source/MultiUserServer/Private/ConcertServerStyle.cpp
@@ -23,7 +23,6 @@ void FConcertServerStyle::Shutdown()
 {
 	if (StyleInstance)
 	{
-		
 		FSlateStyleRegistry::UnRegisterSlateStyle(*StyleInstance);
 		ensure(StyleInstance.IsUnique());
 		StyleInstance.Reset();
@@ -60,17 +59,25 @@ TSharedRef<FSlateStyleSet> FConcertServerStyle::Create()
 	
 	const FVector2D Icon16x16(16.0f, 16.0f); 
 	const FVector2D Icon32x32(32.0f, 32.0f); 
-
+	
 	// Icons
 	StyleSet->Set("Concert.MultiUser", new IMAGE_PLUGIN_BRUSH("Icons/icon_MultiUser_32x", Icon32x32));
 	StyleSet->Set("Concert.SessionContent.ColumnHeader", new IMAGE_PLUGIN_BRUSH_SVG("Icons/Package_16x", Icon16x16));
+	// Icons: Package
 	StyleSet->Set("Concert.SessionContent.PackageAdded", new IMAGE_PLUGIN_BRUSH_SVG("Icons/PackageAdded_16x", Icon16x16));
 	StyleSet->Set("Concert.SessionContent.PackageDeleted", new IMAGE_PLUGIN_BRUSH_SVG("Icons/PackageDeleted_16x", Icon16x16));
 	StyleSet->Set("Concert.SessionContent.PackageRenamed", new IMAGE_PLUGIN_BRUSH_SVG("Icons/PackageRenamed_16x", Icon16x16));
 	StyleSet->Set("Concert.SessionContent.PackageSaved", new IMAGE_PLUGIN_BRUSH_SVG("Icons/PackageSaved_16x", Icon16x16));
+	// Icons: Ack
 	StyleSet->Set("Concert.Ack.Ack", new IMAGE_PLUGIN_BRUSH("Icons/Ack_Ack_16x", Icon16x16));
 	StyleSet->Set("Concert.Ack.Success", new IMAGE_PLUGIN_BRUSH("Icons/Ack_Success_16x", Icon16x16));
 	StyleSet->Set("Concert.Ack.Failure", new IMAGE_PLUGIN_BRUSH("Icons/Ack_Fail_16x", Icon16x16));
+	// Icons: Muting
+	StyleSet->Set("Concert.Muted",     new IMAGE_PLUGIN_BRUSH("Icons/Muted_16x", Icon16x16));
+	StyleSet->Set("Concert.Unmuted",     new IMAGE_PLUGIN_BRUSH("Icons/Unmuted_16x", Icon16x16));
+	// Icons: DependencyView
+	StyleSet->Set("Concert.HighlightHardDependencies",     new IMAGE_PLUGIN_BRUSH("Icons/Ack_Ack_16x", Icon16x16));
+	
 
 	// Clients tab
 	StyleSet->Set("Concert.Clients.DropShadow", new IMAGE_PLUGIN_BOX_BRUSH("ClientThumbnailDropShadow", FMargin(4.0f / 64.0f)));
diff --git a/Engine/Plugins/Developer/Concert/ConcertApp/MultiUserServer/Source/MultiUserServer/Private/Widgets/Clients/Browser/Models/ClientBrowserModel.cpp b/Engine/Plugins/Developer/Concert/ConcertApp/MultiUserServer/Source/MultiUserServer/Private/Widgets/Clients/Browser/Models/ClientBrowserModel.cpp
--- a/Engine/Plugins/Developer/Concert/ConcertApp/MultiUserServer/Source/MultiUserServer/Private/Widgets/Clients/Browser/Models/ClientBrowserModel.cpp
+++ b/Engine/Plugins/Developer/Concert/ConcertApp/MultiUserServer/Source/MultiUserServer/Private/Widgets/Clients/Browser/Models/ClientBrowserModel.cpp
@@ -219,10 +219,15 @@ void UE::MultiUserServer::FClientBrowserModel::AddClientAdminEndpoint(const FCon
 
 void UE::MultiUserServer::FClientBrowserModel::RemoveClientAdminEndpoint(const FConcertEndpointContext& Context)
 {
-	const FMessageNodeId NodeId = EndpointToNodeId[Context.EndpointId];
-	const int32 Index = Clients.IndexOfByPredicate([&NodeId](const TSharedPtr<FClientBrowserItem>& Item)
+	const FMessageNodeId* NodeId = EndpointToNodeId.Find(Context.EndpointId);
+	if (!NodeId)
 	{
-		return Item->MessageNodeId == NodeId;
+		return;
+	}
+	
+	const int32 Index = Clients.IndexOfByPredicate([NodeId](const TSharedPtr<FClientBrowserItem>& Item)
+	{
+		return Item->MessageNodeId == *NodeId;
 	});
 	
 	if (ensure(Clients.IsValidIndex(Index)))
diff --git a/Engine/Plugins/Developer/Concert/ConcertApp/MultiUserServer/Source/MultiUserServer/Private/Widgets/HistoryDeletion/SActivityDependencyView.cpp b/Engine/Plugins/Developer/Concert/ConcertApp/MultiUserServer/Source/MultiUserServer/Private/Widgets/HistoryDeletion/SActivityDependencyView.cpp
deleted file mode 100644
diff --git a/Engine/Plugins/Developer/Concert/ConcertApp/MultiUserServer/Source/MultiUserServer/Private/Widgets/HistoryDeletion/SActivityDependencyView.h b/Engine/Plugins/Developer/Concert/ConcertApp/MultiUserServer/Source/MultiUserServer/Private/Widgets/HistoryDeletion/SActivityDependencyView.h
deleted file mode 100644
diff --git a/Engine/Plugins/Developer/Concert/ConcertApp/MultiUserServer/Source/MultiUserServer/Private/Widgets/HistoryDeletion/SDeleteActivityDependenciesDialog.cpp b/Engine/Plugins/Developer/Concert/ConcertApp/MultiUserServer/Source/MultiUserServer/Private/Widgets/HistoryDeletion/SDeleteActivityDependenciesDialog.cpp
deleted file mode 100644
diff --git a/Engine/Plugins/Developer/Concert/ConcertApp/MultiUserServer/Source/MultiUserServer/Private/Widgets/HistoryDeletion/SDeleteActivityDependenciesDialog.h b/Engine/Plugins/Developer/Concert/ConcertApp/MultiUserServer/Source/MultiUserServer/Private/Widgets/HistoryDeletion/SDeleteActivityDependenciesDialog.h
deleted file mode 100644
diff --git a/Engine/Plugins/Developer/Concert/ConcertApp/MultiUserServer/Source/MultiUserServer/Private/Widgets/HistoryEdition/SActivityDependencyDialog.cpp b/Engine/Plugins/Developer/Concert/ConcertApp/MultiUserServer/Source/MultiUserServer/Private/Widgets/HistoryEdition/SActivityDependencyDialog.cpp
new file mode 100644
--- /dev/null
+++ b/Engine/Plugins/Developer/Concert/ConcertApp/MultiUserServer/Source/MultiUserServer/Private/Widgets/HistoryEdition/SActivityDependencyDialog.cpp
@@ -0,0 +1,130 @@
+// Copyright Epic Games, Inc. All Rights Reserved.
+
+#include "SActivityDependencyDialog.h"
+
+#include "IConcertSyncServer.h"
+#include "MultiUserServerConsoleVariables.h"
+#include "MultiUserServerModule.h"
+#include "SActivityDependencyView.h"
+
+#include "ConcertLogGlobal.h"
+#include "Dialog/SMessageDialog.h"
+#include "HistoryEdition/DebugDependencyGraph.h"
+#include "HistoryEdition/DependencyGraphBuilder.h"
+#include "Window/ModalWindowManager.h"
+#include "Widgets/SBoxPanel.h"
+#include "Widgets/SessionTabs/Archived/ArchivedSessionHistoryController.h"
+#include "Widgets/Text/STextBlock.h"
+
+#define LOCTEXT_NAMESPACE "UnrealMultiUserUI.SActivityDependencyDialog"
+
+namespace UE::MultiUserServer
+{
+	void SActivityDependencyDialog::Construct(const FArguments& InArgs, const FGuid& SessionId, const TSharedRef<IConcertSyncServer>& SyncServer, TSet<FActivityID> InBaseActivities)
+	{
+		check(InArgs._OnConfirmAction.IsBound() && InArgs._AnalyseActivities.IsBound());
+		OnConfirmActionFunc = InArgs._OnConfirmAction;
+		GetFinalInclusionResultTextFunc = InArgs._GetFinalInclusionResultText;
+		ConfirmDialogTitle = InArgs._ConfirmDialogTitle;
+		ConfirmDialogMessage = InArgs._ConfirmDialogMessage;
+		
+		SCustomDialog::Construct(
+			SCustomDialog::FArguments()
+			.Title(InArgs._Title)
+			.Buttons({
+				FButton(InArgs._PerformActionButtonLabel)
+					.SetOnClicked(FSimpleDelegate::CreateSP(this, &SActivityDependencyDialog::OnConfirmPressed)),
+				FButton(LOCTEXT("DeleteActivity.CancelButtonLabel", "Cancel"))
+					.SetOnClicked(FSimpleDelegate::CreateLambda([this]()
+					{
+						RequestDestroyWindow();
+					}))
+					.SetPrimary(true)
+			})
+			.Content()
+			[
+				SNew(SVerticalBox)
+				
+				+SVerticalBox::Slot()
+				.HAlign(HAlign_Fill)
+				.VAlign(VAlign_Top)
+				.AutoHeight()
+				.Padding(0, 5, 0, 10)
+				[
+					SNew(STextBlock)
+					.Text(InArgs._Description)
+				]
+
+				+SVerticalBox::Slot()
+				.HAlign(HAlign_Fill)
+				.VAlign(VAlign_Fill)
+				[
+					CreateBody(InArgs, SessionId, SyncServer, MoveTemp(InBaseActivities))
+				]
+			]
+		);
+	}
+
+	TSharedRef<SWidget> SActivityDependencyDialog::CreateBody(const FArguments& InArgs, const FGuid& InSessionId, const TSharedRef<IConcertSyncServer>& InSyncServer, TSet<FActivityID> InBaseActivities)
+	{
+		const TOptional<FConcertSyncSessionDatabaseNonNullPtr> SessionDatabase = InSyncServer->GetArchivedSessionDatabase(InSessionId);
+		check(SessionDatabase);
+		
+		ConcertSyncCore::FActivityDependencyGraph DependencyGraph = ConcertSyncCore::BuildDependencyGraphFrom(*SessionDatabase);
+		if (ConsoleVariables::CVarLogActivityDependencyGraphOnDelete.GetValueOnGameThread())
+		{
+			UE_LOG(LogConcert, Log, TEXT("Analysed graph: %s"), *UE::ConcertSyncCore::Graphviz::ExportToGraphviz(DependencyGraph, *SessionDatabase));
+		}
+		
+		return SAssignNew(DependencyView, SActivityDependencyView, MoveTemp(DependencyGraph))
+			.AdditionalColumns(InArgs._AdditionalColumns)
+			.BaseActivities(MoveTemp(InBaseActivities))
+			.AnalyseActivities(InArgs._AnalyseActivities)
+			.ShouldShowActivity(InArgs._ShouldShowActivity)
+			.CreateSessionHistoryController_Lambda([this, InSessionId, InSyncServer](const SSessionHistory::FArguments& InArgs)
+			{
+				const TSharedPtr<FArchivedSessionHistoryController> FilteredSessionHistoryController = CreateForDeletionDialog(
+					InSessionId,
+					InSyncServer,
+					InArgs
+					);
+				return FilteredSessionHistoryController.ToSharedRef();
+			})
+			.GetCheckboxToolTip_Lambda([this](const int64 ActivityId, bool bIsHardDependency)
+			{
+				const bool bIsChecked = DependencyView->IsChecked(ActivityId);
+				if (bIsHardDependency)
+				{
+					return FText::Format(LOCTEXT("Action.CheckBox.TooltipDisabledFmt", "Hard dependency. {0}"), GetFinalInclusionResultTextFunc.Execute(ActivityId, bIsChecked));
+				}
+
+				if (bIsChecked)
+				{
+					return FText::Format(LOCTEXT("Action.CheckBox.TooltipEnabled.ApplyActionFmt", "Possible dependency.\n{0}"), GetFinalInclusionResultTextFunc.Execute(ActivityId, bIsChecked));
+				}
+				return FText::Format(LOCTEXT("Action.CheckBox.TooltipEnabled.DoNotApplyActionFmt", "Possible dependency.\n{0}"), GetFinalInclusionResultTextFunc.Execute(ActivityId, bIsChecked));
+			});
+	}
+
+	void SActivityDependencyDialog::OnConfirmPressed()
+	{
+		const TSharedRef<SMessageDialog> Dialog = SNew(SMessageDialog)
+			.Title(ConfirmDialogTitle.IsSet() ? ConfirmDialogTitle.Get() : LOCTEXT("ConfirmDialog.Title", "Confirm action"))
+			.Buttons({
+				FButton(LOCTEXT("ConfirmDialog.Yes", "Yes"))
+					.SetOnClicked(FSimpleDelegate::CreateLambda([KeepAlive = SharedThis(this)]()
+					{
+						KeepAlive->OnConfirmActionFunc.Execute(KeepAlive->DependencyView->GetSelection());
+						KeepAlive->RequestDestroyWindow();
+					})),
+				FButton(LOCTEXT("DeleteActivity.CancelButtonLabel", "Cancel"))
+					.SetPrimary(true)
+			})
+			.Message(ConfirmDialogMessage.IsSet() ? ConfirmDialogMessage.Get() : LOCTEXT("ConfirmDialog.DefaultMessage", "Are you sure?"));
+		FConcertServerUIModule::Get()
+			.GetModalWindowManager()
+			->ShowFakeModalWindow(Dialog);
+	}
+}
+
+#undef LOCTEXT_NAMESPACE
diff --git a/Engine/Plugins/Developer/Concert/ConcertApp/MultiUserServer/Source/MultiUserServer/Private/Widgets/HistoryEdition/SActivityDependencyDialog.h b/Engine/Plugins/Developer/Concert/ConcertApp/MultiUserServer/Source/MultiUserServer/Private/Widgets/HistoryEdition/SActivityDependencyDialog.h
new file mode 100644
--- /dev/null
+++ b/Engine/Plugins/Developer/Concert/ConcertApp/MultiUserServer/Source/MultiUserServer/Private/Widgets/HistoryEdition/SActivityDependencyDialog.h
@@ -0,0 +1,78 @@
+// Copyright Epic Games, Inc. All Rights Reserved.
+
+#pragma once
+
+#include "CoreMinimal.h"
+#include "Dialog/SCustomDialog.h"
+
+#include "ConcertMessageData.h"
+#include "SActivityDependencyView.h"
+#include "Session/Activity/ActivityColumn.h"
+#include "Widgets/DeclarativeSyntaxSupport.h"
+
+class FArchivedSessionHistoryController;
+class IConcertSyncServer;
+
+namespace UE::ConcertSyncCore
+{
+	struct FHistoryAnalysisResult;
+}
+
+namespace UE::MultiUserServer
+{
+	/** Base dialog for perform actions activity dependencies, notably for deleting and muting activities. */
+	class SActivityDependencyDialog : public SCustomDialog
+	{
+	public:
+
+		DECLARE_DELEGATE_OneParam(FConfirmDependencyAction, const TSet<FActivityID>& /*SelectedItems*/)
+		DECLARE_DELEGATE_RetVal_TwoParams(FText, FGetFinalInclusionResultText, const int64 /*ActivityId*/, bool /*bIsChecked*/);
+
+		SLATE_BEGIN_ARGS(SActivityDependencyDialog)
+		{}
+			SLATE_ARGUMENT(FText, Title)
+			SLATE_ARGUMENT(FText, Description)
+			SLATE_ARGUMENT(FText, PerformActionButtonLabel)
+		
+			/** Called when the user confirms the deletion of the activities.*/
+			SLATE_EVENT(FConfirmDependencyAction, OnConfirmAction)
+
+			/** Analyses the given activities */
+			SLATE_EVENT(SActivityDependencyView::FAnalyseActivities, AnalyseActivities)
+		
+			/** Appended to tooltip when hovering the checkbox in the first column. Explains what the current check state means. */
+			SLATE_EVENT(FGetFinalInclusionResultText, GetFinalInclusionResultText)
+
+			/** Optional. Whether to show the activity. */
+			SLATE_EVENT(SActivityDependencyView::FShouldShowActivity, ShouldShowActivity)
+		
+			/** Optional. Additional columns to add on top of the default ones. */
+			SLATE_ARGUMENT(TArray<FActivityColumn>, AdditionalColumns)
+		
+			/** Optional. The title of the dialog that shows up when the user clicks the confirm button. */
+			SLATE_ATTRIBUTE(FText, ConfirmDialogTitle)
+			/** Optional. The content of the dialog that shows up when the user clicks the confirm button. */
+			SLATE_ATTRIBUTE(FText, ConfirmDialogMessage)
+		SLATE_END_ARGS()
+
+		/**
+		 * @param InDependencyArgs Specifies which activities must be deleted and which are optional.
+		 */
+		void Construct(const FArguments& InArgs, const FGuid& SessionId, const TSharedRef<IConcertSyncServer>& SyncServer, TSet<FActivityID> InBaseActivities);
+		
+	private:
+
+		TSharedPtr<SActivityDependencyView> DependencyView;
+		
+		FConfirmDependencyAction OnConfirmActionFunc;
+		FGetFinalInclusionResultText GetFinalInclusionResultTextFunc;
+		
+		TAttribute<FText> ConfirmDialogTitle;
+		TAttribute<FText> ConfirmDialogMessage;
+
+		TSharedRef<SWidget> CreateBody(const FArguments& InArgs, const FGuid& InSessionId, const TSharedRef<IConcertSyncServer>& InSyncServer, TSet<FActivityID> InBaseActivities);
+
+		void OnConfirmPressed();
+	};
+}
+
diff --git a/Engine/Plugins/Developer/Concert/ConcertApp/MultiUserServer/Source/MultiUserServer/Private/Widgets/HistoryEdition/SActivityDependencyView.cpp b/Engine/Plugins/Developer/Concert/ConcertApp/MultiUserServer/Source/MultiUserServer/Private/Widgets/HistoryEdition/SActivityDependencyView.cpp
new file mode 100644
--- /dev/null
+++ b/Engine/Plugins/Developer/Concert/ConcertApp/MultiUserServer/Source/MultiUserServer/Private/Widgets/HistoryEdition/SActivityDependencyView.cpp
@@ -0,0 +1,211 @@
+﻿// Copyright Epic Games, Inc. All Rights Reserved.
+
+#include "SActivityDependencyView.h"
+
+#include "ConcertMessageData.h"
+#include "ConcertServerStyle.h"
+#include "SPrimaryButton.h"
+#include "Algo/AllOf.h"
+#include "HistoryEdition/HistoryAnalysis.h"
+#include "Session/History/AbstractSessionHistoryController.h"
+
+#include "Widgets/Input/SButton.h"
+#include "Widgets/Input/SCheckBox.h"
+#include "Widgets/Layout/SScaleBox.h"
+
+#define LOCTEXT_NAMESPACE "UnrealMultiUserUI.SActivityDependencyView"
+
+namespace UE::MultiUserServer
+{
+	namespace DependencyView
+	{
+		const FName ShouldPerformActivityActionCheckboxColumnID("ShouldPerformActivityActionCheckboxColumnID");
+		DECLARE_DELEGATE_RetVal_OneParam(TSharedRef<SWidget>, FCreateActivityColumnWidget, FActivityID /*ActivityID*/);
+		
+		static FActivityColumn ShouldPerformActivityActionColumn(const FCreateActivityColumnWidget& CreateActivityColumnWidget)
+		{
+			return FActivityColumn(
+				SHeaderRow::Column(ShouldPerformActivityActionCheckboxColumnID)
+					.DefaultLabel(FText::GetEmpty())
+					.FixedWidth(22)
+					.ShouldGenerateWidget(true)
+				)
+				.ColumnSortOrder(static_cast<int32>(ConcertSharedSlate::ActivityColumn::EPredefinedColumnOrder::AvatarColor))
+				.GenerateColumnWidget(FActivityColumn::FGenerateColumnWidget::CreateLambda([CreateActivityColumnWidget](const TSharedRef<SConcertSessionActivities>& Owner, const TSharedRef<FConcertSessionActivity>& Activity, SOverlay::FScopedWidgetSlotArguments& Slot)
+				{
+					Slot
+					.HAlign(HAlign_Center)
+					.VAlign(VAlign_Center)
+					[
+						CreateActivityColumnWidget.Execute(Activity->Activity.ActivityId)
+					];
+				}));
+		}
+	}
+	
+	void SActivityDependencyView::Construct(const FArguments& InArgs, ConcertSyncCore::FActivityDependencyGraph InDependencyGraph)
+	{
+		using namespace DependencyView;
+		check(InArgs._CreateSessionHistoryController.IsBound() && InArgs._AnalyseActivities.IsBound());
+
+		BaseActivities = InArgs._BaseActivities;
+		AnalyseActivitiesDelegate = InArgs._AnalyseActivities;
+		ShouldShowActivity = InArgs._ShouldShowActivity;
+		GetCheckboxToolTipDelegate = InArgs._GetCheckboxToolTip;
+
+		DependencyGraph = MoveTemp(InDependencyGraph);
+		UserSelectedActivities = BaseActivities;
+		Selection = AnalyseActivitiesDelegate.Execute(DependencyGraph, BaseActivities);
+		
+		TSharedPtr<FConcertSessionActivitiesOptions> ViewOptions = MakeShared<FConcertSessionActivitiesOptions>();
+		ViewOptions->bEnableConnectionActivityFiltering = false;
+		ViewOptions->bEnableLockActivityFiltering = false;
+		ViewOptions->bEnableIgnoredActivityFiltering = false;
+
+		const FActivityColumn CheckmarkColumn = ShouldPerformActivityActionColumn(FCreateActivityColumnWidget::CreateLambda([this](FActivityID ActivityID)
+		{
+			const bool bIsHardDependency = Selection.HardDependencies.Contains(ActivityID);
+			return SNew(SCheckBox)
+				.IsEnabled_Lambda([this, ActivityID]()
+				{
+					return !BaseActivities.Contains(ActivityID)
+						&& (UserSelectedActivities.Contains(ActivityID) || Selection.PossibleDependencies.Contains(ActivityID));
+				})
+				.IsChecked_Lambda([this, ActivityID]()
+				{
+					return IsChecked(ActivityID) ? ECheckBoxState::Checked : ECheckBoxState::Unchecked;
+				})
+				.OnCheckStateChanged_Lambda([this, ActivityID](ECheckBoxState NewState)
+				{
+					OnCheckBoxStateChanged(ActivityID, NewState == ECheckBoxState::Checked);
+				})
+				.ToolTipText_Lambda([this, ActivityID, bIsHardDependency]()
+				{
+					return GetCheckboxToolTipDelegate.IsBound() ? GetCheckboxToolTipDelegate.Execute(ActivityID, bIsHardDependency) : FText::GetEmpty();
+				});
+		}));
+		
+		TArray<FActivityColumn> Columns { CheckmarkColumn, ConcertSharedSlate::ActivityColumn::Operation() };
+		for (const FActivityColumn& Column : InArgs._AdditionalColumns)
+		{
+			Columns.Add(Column);
+		}
+
+		SessionHistoryController = InArgs._CreateSessionHistoryController.Execute(
+			SSessionHistory::FArguments()
+				.AllowActivity_Lambda([this](const FConcertSyncActivity& Activity, const TStructOnScope<FConcertSyncActivitySummary>&)
+				{
+					const bool bShouldShow = !ShouldShowActivity.IsBound() || ShouldShowActivity.Execute(Activity);
+					return bShouldShow &&
+						(Selection.HardDependencies.Contains(Activity.ActivityId)
+						|| Selection.PossibleDependencies.Contains(Activity.ActivityId));
+				})
+				.Columns(Columns)
+				.ViewOptions(ViewOptions)
+				.DetailsAreaVisibility(EVisibility::Collapsed)
+				.DarkenMutedActivities(false)
+				.OnMakeColumnOverlayWidget(this, &SActivityDependencyView::CreateCheckboxOverlay)
+			);
+		ChildSlot
+		[
+			SessionHistoryController->GetSessionHistory()
+		];
+	}
+
+	bool SActivityDependencyView::IsChecked(FActivityID ActivityID) const
+	{
+		return UserSelectedActivities.Contains(ActivityID) || Selection.HardDependencies.Contains(ActivityID);
+	}
+
+	TSet<FActivityID> SActivityDependencyView::GetSelection() const
+	{
+		TSet<FActivityID> Result = UserSelectedActivities;
+		for (const FActivityID HardDependency : Selection.HardDependencies)
+		{
+			Result.Add(HardDependency);
+		}
+		return Result;
+	}
+
+	void SActivityDependencyView::OnCheckBoxStateChanged(const FActivityID ActivityId, bool bIsChecked)
+	{
+		check(!BaseActivities.Contains(ActivityId));
+		if (bIsChecked)
+		{
+			UserSelectedActivities.Add(ActivityId);
+		}
+		else
+		{
+			UserSelectedActivities.Remove(ActivityId);
+		}
+
+		Selection = AnalyseActivitiesDelegate.Execute(DependencyGraph, UserSelectedActivities);
+		SessionHistoryController->ReloadActivities();
+	}
+
+	TSharedPtr<SWidget> SActivityDependencyView::CreateCheckboxOverlay(TWeakPtr<SMultiColumnTableRow<TSharedPtr<FConcertSessionActivity>>> Row, TWeakPtr<FConcertSessionActivity> Activity, const FName& ColumnName)
+	{
+		TSharedPtr<FConcertSessionActivity> ActivityPin = Activity.Pin();
+		check(ActivityPin);
+
+		const FActivityID ActivityID = ActivityPin->Activity.ActivityId;
+		const bool bCanEverCreate = !BaseActivities.Contains(ActivityID);
+		if (ColumnName != DependencyView::ShouldPerformActivityActionCheckboxColumnID || !bCanEverCreate)
+		{
+			return nullptr;
+		}
+
+		return SNew(SButton)
+			.ButtonStyle(&FAppStyle::Get().GetWidgetStyle<FButtonStyle>("PrimaryButton"))
+			.ForegroundColor(FSlateColor::UseStyle())
+			.HAlign(HAlign_Center)
+			.VAlign(VAlign_Center)
+			.Visibility_Lambda([this, ActivityID, Row]()
+			{
+				return Row.Pin()->IsHovered() && CanRemoveMandatorySelection(ActivityID) ? EVisibility::Visible : EVisibility::Hidden;
+			})
+			.ToolTipText(LOCTEXT("MandatorilyChecked", "Click to highlight the activities that mandate this activity's selection."))
+			.OnClicked_Lambda([this, ActivityID]()
+			{
+				const TOptional<ConcertSyncCore::FActivityNodeID> NodeId = DependencyGraph.FindNodeByActivity(ActivityID);
+				if (ensure(NodeId))
+				{
+					const ConcertSyncCore::FActivityNode& Node = DependencyGraph.GetNodeById(*NodeId);
+					const TArray<TSharedPtr<FConcertSessionActivity>>& Activities = SessionHistoryController->GetSessionHistory()->GetActivities();
+					TArray<TSharedPtr<FConcertSessionActivity>> NewItemSelection;
+					for (const TSharedPtr<FConcertSessionActivity>& Activity : Activities)
+					{
+						if (Node.DependsOnActivity(Activity->Activity.ActivityId, DependencyGraph, {}, ConcertSyncCore::EDependencyStrength::HardDependency))
+						{
+							NewItemSelection.Add(Activity);
+						}
+					}
+					SessionHistoryController->GetSessionHistory()->SetSelectedActivities(NewItemSelection);
+				}
+				return FReply::Handled();
+			})
+			.ContentPadding(0)
+			[
+				SNew(SScaleBox)
+				[
+					SNew(SImage)
+					.ColorAndOpacity(FSlateColor::UseForeground())
+					.Image(FConcertServerStyle::Get().GetBrush("Concert.HighlightHardDependencies"))
+				]
+			];
+	}
+
+	bool SActivityDependencyView::CanRemoveMandatorySelection(const FActivityID ActivityID) const
+	{
+		const TOptional<ConcertSyncCore::FActivityNodeID> NodeId = DependencyGraph.FindNodeByActivity(ActivityID);
+		const bool bCanRemoveMandatorySelection = ensure(NodeId)
+			&& Algo::AllOf(DependencyGraph.GetNodeById(*NodeId).GetDependencies(), [this](const ConcertSyncCore::FActivityDependencyEdge& DependencyEdge)
+			{
+				const FActivityID TargetActivityId = DependencyGraph.GetNodeById(DependencyEdge.GetDependedOnNodeID()).GetActivityId();
+				return !BaseActivities.Contains(TargetActivityId);
+			});
+		return IsChecked(ActivityID) && !UserSelectedActivities.Contains(ActivityID) && bCanRemoveMandatorySelection;
+	}
+}
+
+#undef LOCTEXT_NAMESPACE 
diff --git a/Engine/Plugins/Developer/Concert/ConcertApp/MultiUserServer/Source/MultiUserServer/Private/Widgets/HistoryEdition/SActivityDependencyView.h b/Engine/Plugins/Developer/Concert/ConcertApp/MultiUserServer/Source/MultiUserServer/Private/Widgets/HistoryEdition/SActivityDependencyView.h
new file mode 100644
--- /dev/null
+++ b/Engine/Plugins/Developer/Concert/ConcertApp/MultiUserServer/Source/MultiUserServer/Private/Widgets/HistoryEdition/SActivityDependencyView.h
@@ -0,0 +1,81 @@
+﻿// Copyright Epic Games, Inc. All Rights Reserved.
+
+#pragma once
+
+#include "CoreMinimal.h"
+#include "HistoryEdition/ActivityDependencyGraph.h"
+#include "HistoryEdition/HistoryAnalysis.h"
+#include "Session/History/SSessionHistory.h"
+#include "Widgets/DeclarativeSyntaxSupport.h"
+#include "Widgets/SCompoundWidget.h"
+
+class FAbstractSessionHistoryController;
+template<typename T> class SMultiColumnTableRow;
+struct FConcertSyncActivity;
+
+namespace UE::MultiUserServer
+{
+	/**
+	 * A view for displaying activity dependencies.
+	 * Each row has a checkbox. Useful for actions, such as deleting, muting, and unmuting.
+	 * Adds additional dependencies depending on what optional dependencies are checked.
+	 */
+	class SActivityDependencyView : public SCompoundWidget
+	{
+	public:
+
+		DECLARE_DELEGATE_RetVal_OneParam(TSharedRef<FAbstractSessionHistoryController>, FCreateSessionHistoryController, const SSessionHistory::FArguments& /*InArgs*/)
+		DECLARE_DELEGATE_RetVal_TwoParams(ConcertSyncCore::FHistoryAnalysisResult, FAnalyseActivities, const ConcertSyncCore::FActivityDependencyGraph&, const TSet<FActivityID>& /*Activities*/);
+		DECLARE_DELEGATE_RetVal_OneParam(bool, FShouldShowActivity, const FConcertSyncActivity /*ActivityId*/);
+		DECLARE_DELEGATE_RetVal_TwoParams(FText, FGetCheckboxToolTip, const int64& /*Activity*/, bool /*bIsHardDependency*/)
+
+		SLATE_BEGIN_ARGS(SActivityDependencyView)
+		{}
+			/** The activities this view is based on. Always checked. */
+			SLATE_ARGUMENT(TSet<FActivityID>, BaseActivities)
+		
+			/** Required. A function that will create a session history widget. */
+			SLATE_EVENT(FCreateSessionHistoryController, CreateSessionHistoryController)
+			/** Analyses the given activities */
+			SLATE_EVENT(FAnalyseActivities, AnalyseActivities)
+
+			/** Optional. Whether to show the activity */
+			SLATE_EVENT(FShouldShowActivity, ShouldShowActivity)
+			/** Optional. The tooltip to display when the user hovers a */
+			SLATE_EVENT(FGetCheckboxToolTip, GetCheckboxToolTip)
+		
+			/** Optional. Additional columns to add on top of the default ones. */
+			SLATE_ARGUMENT(TArray<FActivityColumn>, AdditionalColumns)
+		SLATE_END_ARGS()
+
+		void Construct(const FArguments& InArgs, ConcertSyncCore::FActivityDependencyGraph InDependencyGraph);
+
+		/** @return Whether the given activity ID is checked */
+		bool IsChecked(FActivityID ActivityID) const;
+		TSet<FActivityID> GetSelection() const;
+		
+	private:
+
+		/** Used to render the rows */
+		TSharedPtr<FAbstractSessionHistoryController> SessionHistoryController;
+		
+		/** Built once on construction. */
+		ConcertSyncCore::FActivityDependencyGraph DependencyGraph;
+		
+		/** The activities this view is based on. Always checked. */
+		TSet<FActivityID> BaseActivities;
+		TSet<FActivityID> UserSelectedActivities;
+
+		/** The current selection. Updated every time something is checked. */
+		ConcertSyncCore::FHistoryAnalysisResult Selection;
+
+		FAnalyseActivities AnalyseActivitiesDelegate;
+		FShouldShowActivity ShouldShowActivity;
+		FGetCheckboxToolTip GetCheckboxToolTipDelegate;
+
+		void OnCheckBoxStateChanged(const FActivityID ActivityId, bool bIsChecked);
+		TSharedPtr<SWidget> CreateCheckboxOverlay(TWeakPtr<SMultiColumnTableRow<TSharedPtr<FConcertSessionActivity>>> Row, TWeakPtr<FConcertSessionActivity> Activity, const FName& ColumnName);
+
+		bool CanRemoveMandatorySelection(const FActivityID ActivityID) const;
+	};
+}
diff --git a/Engine/Plugins/Developer/Concert/ConcertApp/MultiUserServer/Source/MultiUserServer/Private/Widgets/HistoryEdition/SDeleteActivityDependenciesDialog.cpp b/Engine/Plugins/Developer/Concert/ConcertApp/MultiUserServer/Source/MultiUserServer/Private/Widgets/HistoryEdition/SDeleteActivityDependenciesDialog.cpp
new file mode 100644
--- /dev/null
+++ b/Engine/Plugins/Developer/Concert/ConcertApp/MultiUserServer/Source/MultiUserServer/Private/Widgets/HistoryEdition/SDeleteActivityDependenciesDialog.cpp
@@ -0,0 +1,30 @@
+﻿// Copyright Epic Games, Inc. All Rights Reserved.
+
+#include "SDeleteActivityDependenciesDialog.h"
+
+#define LOCTEXT_NAMESPACE "UnrealMultiUserUI.SDeleteActivityDependenciesDialog"
+
+void UE::MultiUserServer::SDeleteActivityDependenciesDialog::Construct(const FArguments& InArgs, const FGuid& SessionId, const TSharedRef<IConcertSyncServer>& SyncServer, TSet<FActivityID> InBaseActivities)
+{
+	SActivityDependencyDialog::Construct(
+		SActivityDependencyDialog::FArguments()
+			.Title(LOCTEXT("RemoveActivityTitle", "Remove activity"))
+			.Description(LOCTEXT("DeleteActivity.HeaderBaseText", "Review the activities that will be deleted:"))
+			.PerformActionButtonLabel(LOCTEXT("DeleteActivity.ConfirmButtonLabel", "Delete"))
+			.OnConfirmAction(InArgs._OnConfirmDeletion)
+			.AnalyseActivities_Lambda([](const ConcertSyncCore::FActivityDependencyGraph& Graph, const TSet<FActivityID>& ActivityIds)
+			{
+				return ConcertSyncCore::AnalyseActivityDependencies_TopDown(ActivityIds, Graph, true);
+			})
+			.GetFinalInclusionResultText_Lambda([](const int64 ActivityId, bool bIsChecked)
+			{
+				return bIsChecked
+					? LOCTEXT("WillDelete", "Will delete.")
+					: LOCTEXT("WillNotDelete", "Won't delete.");
+			})
+			.ConfirmDialogTitle(LOCTEXT("ConfirmDeletion.Title", "Confirm deletion"))
+			.ConfirmDialogMessage(LOCTEXT("ConfirmDeletion.Description", "Are you sure you want to delete the selected activities?")),
+		SessionId, SyncServer, MoveTemp(InBaseActivities)
+	);
+}
+#undef LOCTEXT_NAMESPACE
diff --git a/Engine/Plugins/Developer/Concert/ConcertApp/MultiUserServer/Source/MultiUserServer/Private/Widgets/HistoryEdition/SDeleteActivityDependenciesDialog.h b/Engine/Plugins/Developer/Concert/ConcertApp/MultiUserServer/Source/MultiUserServer/Private/Widgets/HistoryEdition/SDeleteActivityDependenciesDialog.h
new file mode 100644
--- /dev/null
+++ b/Engine/Plugins/Developer/Concert/ConcertApp/MultiUserServer/Source/MultiUserServer/Private/Widgets/HistoryEdition/SDeleteActivityDependenciesDialog.h
@@ -0,0 +1,34 @@
+﻿// Copyright Epic Games, Inc. All Rights Reserved.
+
+#pragma once
+
+#include "CoreMinimal.h"
+#include "SActivityDependencyDialog.h"
+#include "Widgets/DeclarativeSyntaxSupport.h"
+
+class FArchivedSessionHistoryController;
+namespace UE::ConcertSyncCore
+{
+	struct FHistoryAnalysisResult;
+}
+
+namespace UE::MultiUserServer
+{
+	/** Displayed when a user asks to delete an activity. */
+	class SDeleteActivityDependenciesDialog : public SActivityDependencyDialog
+	{
+	public:
+		SLATE_BEGIN_ARGS(SDeleteActivityDependenciesDialog)
+		{}
+			/** Called when the user confirms the deletion of the activities.*/
+			SLATE_EVENT(SActivityDependencyDialog::FConfirmDependencyAction, OnConfirmDeletion)
+		SLATE_END_ARGS()
+
+		/**
+		 * @param InDeletionRequirements Specifies which activities must be deleted and which are optional.
+		 */
+		void Construct(const FArguments& InArgs, const FGuid& SessionId, const TSharedRef<IConcertSyncServer>& SyncServer, TSet<FActivityID> InBaseActivities);
+	};
+}
+
+
diff --git a/Engine/Plugins/Developer/Concert/ConcertApp/MultiUserServer/Source/MultiUserServer/Private/Widgets/HistoryEdition/SMuteActivityDependenciesDialog.cpp b/Engine/Plugins/Developer/Concert/ConcertApp/MultiUserServer/Source/MultiUserServer/Private/Widgets/HistoryEdition/SMuteActivityDependenciesDialog.cpp
new file mode 100644
--- /dev/null
+++ b/Engine/Plugins/Developer/Concert/ConcertApp/MultiUserServer/Source/MultiUserServer/Private/Widgets/HistoryEdition/SMuteActivityDependenciesDialog.cpp
@@ -0,0 +1,121 @@
+// Copyright Epic Games, Inc. All Rights Reserved.
+
+#include "SMuteActivityDependenciesDialog.h"
+
+#include "ConcertServerStyle.h"
+#include "ConcertSyncSessionDatabase.h"
+#include "IConcertSyncServer.h"
+#include "HistoryEdition/HistoryAnalysis.h"
+#include "Session/Activity/PredefinedActivityColumns.h"
+
+#define LOCTEXT_NAMESPACE "UnrealMultiUserUI.SMuteActivityDependenciesDialog"
+
+namespace UE::MultiUserServer
+{
+	namespace MuteActivityDependenciesDialog
+	{
+		const FName IsMutedColumnID("IsMutedColumn");
+
+		DECLARE_DELEGATE_RetVal_OneParam(TSharedRef<SWidget>, FCreateActivityColumnWidget, FActivityID /*ActivityID*/);
+	
+		static FActivityColumn IsMutedColumn(const FCreateActivityColumnWidget& CreateActivityColumnWidget)
+		{
+			return FActivityColumn(
+				SHeaderRow::Column(IsMutedColumnID)
+					.DefaultLabel(FText::GetEmpty())
+					.FixedWidth(20)
+					.ShouldGenerateWidget(true)
+				)
+				// Place it right after the checkbox, see ShouldPerformActivityActionColumn in SActivityDependencyDialog.cpp
+				.ColumnSortOrder(static_cast<int32>(ConcertSharedSlate::ActivityColumn::EPredefinedColumnOrder::AvatarColor) + 1)
+				.GenerateColumnWidget(FActivityColumn::FGenerateColumnWidget::CreateLambda([CreateActivityColumnWidget](const TSharedRef<SConcertSessionActivities>& Owner, const TSharedRef<FConcertSessionActivity>& Activity, SOverlay::FScopedWidgetSlotArguments& Slot)
+				{
+					Slot
+					.VAlign(VAlign_Center)
+					[
+						CreateActivityColumnWidget.Execute(Activity->Activity.ActivityId)
+					];
+				}));
+		}
+	}
+
+	void SMuteActivityDependenciesDialog::Construct(const FArguments& InArgs, const FGuid& InSessionId, const TSharedRef<IConcertSyncServer>& InSyncServer, TSet<FActivityID> InBaseActivities)
+	{
+		SyncServer = InSyncServer;
+		SessionId = InSessionId;
+		
+		using namespace MuteActivityDependenciesDialog;
+		const FActivityColumn MutedColumn = IsMutedColumn(FCreateActivityColumnWidget::CreateLambda([this](FActivityID ActivityID)
+		{
+			const bool bIsMuted = IsMuted(ActivityID);
+			return SNew(SImage)
+				.Image(bIsMuted ? FConcertServerStyle::Get().GetBrush("Concert.Muted") : FConcertServerStyle::Get().GetBrush("Concert.Unmuted"))
+				.ToolTipText(bIsMuted ? LOCTEXT("Muted", "Muted. Server does not send this activity to clients.") : LOCTEXT("Unmuted", "Not muted. Server sends this activity to clients."));
+		}));
+
+		const bool bShouldMute = InArgs._MuteOperation == EMuteOperation::Mute;
+		const FText DialogTitle = bShouldMute
+			? LOCTEXT("Mute.Title", "Mute activities")
+			: LOCTEXT("Unmute.Title", "Unmute activities");
+		const FText Description = bShouldMute
+			? LOCTEXT("Mute.Description", "Review the activities that will be muted.\nMuted activities are not sent to clients.")
+			: LOCTEXT("Unmute.Description", "Review the activities that will be unmuted.\nMuted activities are not sent to clients.");
+		const FText ButtonLabel = bShouldMute
+			? LOCTEXT("Mute.ButtonLabel", "Mute")
+			: LOCTEXT("Unmute.ButtonLabel", "Unmute");
+		const FText ConfirmDialogTitleArg = bShouldMute
+			? LOCTEXT("Mute.ConfirmDialogTitle", "Confirm mute")
+			: LOCTEXT("Unmute.ConfirmDialogTitle", "Confirm unmute");
+		const FText ConfirmDialogDescriptionArg = bShouldMute
+			? LOCTEXT("Mute.ConfirmDialogDescription", "Are you sure you want to mute the selected activities?")
+			: LOCTEXT("Unmute.ConfirmDialogDescription", "Are you sure you want to unmute the selected activities?");
+		
+		SActivityDependencyDialog::Construct(
+			SActivityDependencyDialog::FArguments()
+				.Title(DialogTitle)
+				.Description(Description)
+				.PerformActionButtonLabel(ButtonLabel)
+				.OnConfirmAction(InArgs._OnConfirmMute)
+				.AdditionalColumns({ MutedColumn })
+				.AnalyseActivities_Lambda([bShouldMute](const ConcertSyncCore::FActivityDependencyGraph& Graph, const TSet<FActivityID>& ActivityIds)
+				{
+					return bShouldMute
+						? ConcertSyncCore::AnalyseActivityDependencies_TopDown(ActivityIds, Graph, true)
+						: ConcertSyncCore::AnalyseActivityDependencies_BottomUp(ActivityIds, Graph, true);
+				})
+				// Improved UX: Don't show  muted activities when muting and unmuted activities when unmuting
+				.ShouldShowActivity_Lambda([bShouldMute](const FConcertSyncActivity& Activity)
+				{
+					const bool bIsMuted = (Activity.Flags & EConcertSyncActivityFlags::Muted) != EConcertSyncActivityFlags::None;
+					return bIsMuted != bShouldMute;
+				})
+				.GetFinalInclusionResultText_Lambda([this, bShouldMute](const int64 ActivityId, bool bIsChecked)
+				{
+					return bShouldMute
+						? LOCTEXT("WillMute", "Will be muted.")
+						: LOCTEXT("WillUnmute", "Will be unmuted.");
+				})
+				.ConfirmDialogTitle(ConfirmDialogTitleArg)
+				.ConfirmDialogMessage(ConfirmDialogDescriptionArg),
+			InSessionId, InSyncServer, MoveTemp(InBaseActivities)
+		);
+	}
+
+	bool SMuteActivityDependenciesDialog::IsMuted(const int64 ActivityId)
+	{
+		if (TSharedPtr<IConcertSyncServer> PinnedSyncServer = SyncServer.Pin())
+		{
+			FConcertSyncActivity Activity;
+			const bool bGotActivity = PinnedSyncServer->GetArchivedSessionDatabase(SessionId)->GetActivity(ActivityId, Activity);
+			check(bGotActivity);
+			
+			const bool bIsMuted = (Activity.Flags & EConcertSyncActivityFlags::Muted) != EConcertSyncActivityFlags::None;
+			return bIsMuted;
+		}
+
+		// Main window might have been closed destroying the server...
+		return false;
+	}
+}
+
+#undef LOCTEXT_NAMESPACE
diff --git a/Engine/Plugins/Developer/Concert/ConcertApp/MultiUserServer/Source/MultiUserServer/Private/Widgets/HistoryEdition/SMuteActivityDependenciesDialog.h b/Engine/Plugins/Developer/Concert/ConcertApp/MultiUserServer/Source/MultiUserServer/Private/Widgets/HistoryEdition/SMuteActivityDependenciesDialog.h
new file mode 100644
--- /dev/null
+++ b/Engine/Plugins/Developer/Concert/ConcertApp/MultiUserServer/Source/MultiUserServer/Private/Widgets/HistoryEdition/SMuteActivityDependenciesDialog.h
@@ -0,0 +1,51 @@
+// Copyright Epic Games, Inc. All Rights Reserved.
+
+#pragma once
+
+#include "CoreMinimal.h"
+#include "SActivityDependencyDialog.h"
+#include "Widgets/DeclarativeSyntaxSupport.h"
+
+class FArchivedSessionHistoryController;
+namespace UE::ConcertSyncCore
+{
+	struct FHistoryAnalysisResult;
+}
+
+namespace UE::MultiUserServer
+{
+	/** Displayed when a user asks to delete an activity. */
+	class SMuteActivityDependenciesDialog : public SActivityDependencyDialog
+	{
+	public:
+		enum class EMuteOperation
+		{
+			Mute,
+			Unmute
+		};
+		
+		SLATE_BEGIN_ARGS(SMuteActivityDependenciesDialog)
+			: _MuteOperation(EMuteOperation::Mute)
+		{}
+			/** Called when the user confirms the muting of the activities.*/
+			SLATE_EVENT(SActivityDependencyDialog::FConfirmDependencyAction, OnConfirmMute)
+
+			/** Whether the activities should be muted or unmuted */
+			SLATE_ARGUMENT(EMuteOperation, MuteOperation)
+		SLATE_END_ARGS()
+
+		/**
+		 * @param InDependencyArgs Specifies which activities must be muted and which are optional.
+		 */
+		void Construct(const FArguments& InArgs, const FGuid& InSessionId, const TSharedRef<IConcertSyncServer>& InSyncServer, TSet<FActivityID> InBaseActivities);
+
+	private:
+
+		TWeakPtr<IConcertSyncServer> SyncServer;
+		FGuid SessionId;
+
+		bool IsMuted(const int64 ActivityId);
+	};
+}
+
+
diff --git a/Engine/Plugins/Developer/Concert/ConcertApp/MultiUserServer/Source/MultiUserServer/Private/Widgets/SessionTabs/Archived/ArchivedConcertSessionTab.cpp b/Engine/Plugins/Developer/Concert/ConcertApp/MultiUserServer/Source/MultiUserServer/Private/Widgets/SessionTabs/Archived/ArchivedConcertSessionTab.cpp
--- a/Engine/Plugins/Developer/Concert/ConcertApp/MultiUserServer/Source/MultiUserServer/Private/Widgets/SessionTabs/Archived/ArchivedConcertSessionTab.cpp
+++ b/Engine/Plugins/Developer/Concert/ConcertApp/MultiUserServer/Source/MultiUserServer/Private/Widgets/SessionTabs/Archived/ArchivedConcertSessionTab.cpp
@@ -4,25 +4,24 @@
 
 #include "ArchivedSessionHistoryController.h"
 #include "ConcertLogGlobal.h"
+#include "ConcertSyncSessionDatabase.h"
 #include "HistoryEdition/ActivityNode.h"
-#include "HistoryEdition/DebugDependencyGraph.h"
-#include "HistoryEdition/DependencyGraphBuilder.h"
 #include "HistoryEdition/HistoryAnalysis.h"
 #include "HistoryEdition/HistoryEdition.h"
-#include "MultiUserServerConsoleVariables.h"
 #include "MultiUserServerModule.h"
 #include "Session/History/SEditableSessionHistory.h"
 #include "Session/History/SSessionHistory.h"
 #include "Window/ModalWindowManager.h"
-#include "Widgets/HistoryDeletion/SDeleteActivityDependenciesDialog.h"
+#include "Widgets/HistoryEdition/SDeleteActivityDependenciesDialog.h"
 #include "Widgets/SessionTabs/Archived/SConcertArchivedSessionTabView.h"
+#include "Widgets/HistoryEdition/SMuteActivityDependenciesDialog.h"
 
 #include "Algo/AllOf.h"
 #include "Algo/Transform.h"
 #include "Dialog/SMessageDialog.h"
 #include "Widgets/Docking/SDockTab.h"
 
-#define LOCTEXT_NAMESPACE "UnrealMultiUserUI"
+#define LOCTEXT_NAMESPACE "UnrealMultiUserUI.FArchivedConcertSessionTab"
 
 FArchivedConcertSessionTab::FArchivedConcertSessionTab(FGuid InspectedSessionID, TSharedRef<IConcertSyncServer> SyncServer, TAttribute<TSharedRef<SWindow>> ConstructUnderWindow)
 	: FConcertSessionTabBase(InspectedSessionID, SyncServer)
@@ -45,38 +44,36 @@ void FArchivedConcertSessionTab::CreateDockContent(const TSharedRef<SDockTab>& I
 			.ConstructUnderMajorTab(InDockTab)
 			.ConstructUnderWindow(ConstructUnderWindow.Get())
 			.MakeSessionHistory(MoveTemp(MakeSessionHistory))
-			.DeleteActivity_Raw(this, &FArchivedConcertSessionTab::OnRequestDeleteActivity)
-			.CanDeleteActivity_Raw(this, &FArchivedConcertSessionTab::CanDeleteActivity)
+			.DeleteActivities_Raw(this, &FArchivedConcertSessionTab::OnRequestDeleteActivities)
+			.CanDeleteActivities_Raw(this, &FArchivedConcertSessionTab::CanDeleteActivities)
+			.MuteActivities_Raw(this, &FArchivedConcertSessionTab::OnRequestMuteActivities)
+			.CanMuteActivities_Raw(this, &FArchivedConcertSessionTab::CanMuteActivities)
+			.UnmuteActivities_Raw(this, &FArchivedConcertSessionTab::OnRequestUnmuteActivities)
+			.CanUnmuteActivities_Raw(this, &FArchivedConcertSessionTab::CanUnmuteActivities)
 		);
 }
 
-void FArchivedConcertSessionTab::OnRequestDeleteActivity(const TSet<TSharedRef<FConcertSessionActivity>>& ActivitiesToDelete) const
+void FArchivedConcertSessionTab::OnRequestDeleteActivities(const TSet<TSharedRef<FConcertSessionActivity>>& ActivitiesToDelete) const
 {
 	using namespace UE::ConcertSyncCore;
+	using namespace UE::MultiUserServer;
 	
 	if (const TOptional<FConcertSyncSessionDatabaseNonNullPtr> SessionDatabase = SyncServer->GetArchivedSessionDatabase(InspectedSessionID))
 	{
-		const FActivityDependencyGraph DependencyGraph = BuildDependencyGraphFrom(*SessionDatabase);
-		if (UE::MultiUserServer::ConsoleVariables::CVarLogActivityDependencyGraphOnDelete.GetValueOnGameThread())
-		{
-			UE_LOG(LogConcert, Log, TEXT("%s"), *UE::ConcertSyncCore::Graphviz::ExportToGraphviz(DependencyGraph, *SessionDatabase));
-		}
-		
 		TSet<FActivityID> RequestedForDelete;
 		Algo::Transform(ActivitiesToDelete, RequestedForDelete, [](const TSharedRef<FConcertSessionActivity>& Activity)
 		{
 			return Activity->Activity.ActivityId;
 		});
-		FHistoryAnalysisResult DeletionRequirements = AnalyseActivityDependencies_TopDown(RequestedForDelete, DependencyGraph, true);
 
-		TWeakPtr<const FArchivedConcertSessionTab> WeakTabThis = SharedThis(this);
-		TSharedRef<SDeleteActivityDependenciesDialog> Dialog = SNew(SDeleteActivityDependenciesDialog, InspectedSessionID, SyncServer, MoveTemp(DeletionRequirements))
-			.OnConfirmDeletion_Lambda([WeakTabThis](const FHistoryAnalysisResult& SelectedRequirements)
+		const TWeakPtr<const FArchivedConcertSessionTab> WeakTabThis = SharedThis(this);
+		const TSharedRef<SDeleteActivityDependenciesDialog> Dialog = SNew(SDeleteActivityDependenciesDialog, InspectedSessionID, SyncServer, MoveTemp(RequestedForDelete))
+			.OnConfirmDeletion_Lambda([WeakTabThis](const TSet<FActivityID>& Selection)
 			{
 				// Because the dialog is non-modal, the user may have closed the program in the mean time
 				if (const TSharedPtr<const FArchivedConcertSessionTab> PinnedThis = WeakTabThis.Pin())
 				{
-					const FOperationErrorResult ErrorResult = DeleteActivitiesInArchivedSession(PinnedThis->SyncServer->GetConcertServer(), PinnedThis->InspectedSessionID, CombineRequirements(SelectedRequirements));
+					const FOperationErrorResult ErrorResult = DeleteActivitiesInArchivedSession(PinnedThis->SyncServer->GetConcertServer(), PinnedThis->InspectedSessionID, Selection);
 					if (ErrorResult.HadError())
 					{
 						UE_LOG(LogConcert, Error, TEXT("Failed to delete activities from session %s: %s"), *PinnedThis->InspectedSessionID.ToString(), *ErrorResult.ErrorMessage->ToString());
@@ -98,13 +95,13 @@ void FArchivedConcertSessionTab::OnRequestDeleteActivity(const TSet<TSharedRef<F
 				}
 			});
 		
-		UE::MultiUserServer::FConcertServerUIModule::Get()
+		FConcertServerUIModule::Get()
 			.GetModalWindowManager()
 			->ShowFakeModalWindow(Dialog);
 	}
 }
 
-FCanDeleteActivitiesResult FArchivedConcertSessionTab::CanDeleteActivity(const TSet<TSharedRef<FConcertSessionActivity>>& ActivitiesToDelete) const
+FCanPerformActionResult FArchivedConcertSessionTab::CanDeleteActivities(const TSet<TSharedRef<FConcertSessionActivity>>& ActivitiesToDelete) const
 {
 	const bool bOnlyPackagesAndTransactions =  Algo::AllOf(ActivitiesToDelete, [](const TSharedRef<FConcertSessionActivity>& Activity)
 	{
@@ -112,10 +109,127 @@ FCanDeleteActivitiesResult FArchivedConcertSessionTab::CanDeleteActivity(const T
 	});
 	if (!bOnlyPackagesAndTransactions)
 	{
-		return FCanDeleteActivitiesResult::No(LOCTEXT("CanDeleteActivity.OnlyPackagesAndTransactionsReason", "Only package and transaction activities can be deleted (the current selection includes other activity types)."));
+		return FCanPerformActionResult::No(LOCTEXT("Delete.CanDeleteActivity.OnlyPackagesAndTransactionsReason", "Only package and transaction activities can be deleted (the current selection includes other activity types)."));
+	}
+
+	return FCanPerformActionResult::Yes();
+}
+
+void FArchivedConcertSessionTab::OnRequestMuteActivities(const TSet<TSharedRef<FConcertSessionActivity>>& ActivitiesToMute) const
+{
+	constexpr bool bShouldMute = true;
+	SpawnDialogForMutingOrUnmuting(bShouldMute, ActivitiesToMute);
+}
+
+FCanPerformActionResult FArchivedConcertSessionTab::CanMuteActivities(const TSet<TSharedRef<FConcertSessionActivity>>& ActivitiesToDelete) const
+{
+	const bool bOnlyPackagesAndTransactions = Algo::AllOf(ActivitiesToDelete, [](const TSharedRef<FConcertSessionActivity>& Activity)
+	{
+		return Activity->Activity.EventType == EConcertSyncActivityEventType::Package || Activity->Activity.EventType == EConcertSyncActivityEventType::Transaction;
+	});
+	if (!bOnlyPackagesAndTransactions)
+	{
+		return FCanPerformActionResult::No(LOCTEXT("Mute.CanMuteActivity.OnlyPackagesAndTransactionsReason", "Only package and transaction activities can be deleted (the current selection includes other activity types)."));
+	}
+
+	const bool bAreAllUnmuted = Algo::AllOf(ActivitiesToDelete, [](const TSharedRef<FConcertSessionActivity>& Activity)
+	{
+		return (Activity->Activity.Flags & EConcertSyncActivityFlags::Muted) == EConcertSyncActivityFlags::None;
+	});
+	if (!bAreAllUnmuted)
+	{
+		return FCanPerformActionResult::No(LOCTEXT("Mute.CanMuteActivity.NotAllUnmuted", "Some of the selected activities are muted."));
+	}
+
+	return FCanPerformActionResult::Yes();
+}
+
+void FArchivedConcertSessionTab::OnRequestUnmuteActivities(const TSet<TSharedRef<FConcertSessionActivity>>& ActivitiesToUnmute) const
+{
+	constexpr bool bShouldMute = false;
+	SpawnDialogForMutingOrUnmuting(bShouldMute, ActivitiesToUnmute);
+}
+
+FCanPerformActionResult FArchivedConcertSessionTab::CanUnmuteActivities(const TSet<TSharedRef<FConcertSessionActivity>>& ActivitiesToDelete) const
+{
+	const bool bOnlyPackagesAndTransactions = Algo::AllOf(ActivitiesToDelete, [](const TSharedRef<FConcertSessionActivity>& Activity)
+	{
+		return Activity->Activity.EventType == EConcertSyncActivityEventType::Package || Activity->Activity.EventType == EConcertSyncActivityEventType::Transaction;
+	});
+	if (!bOnlyPackagesAndTransactions)
+	{
+		return FCanPerformActionResult::No(LOCTEXT("Mute.CanMuteActivity.OnlyPackagesAndTransactionsReason", "Only package and transaction activities can be deleted (the current selection includes other activity types)."));
+	}
+
+	const bool bAreAllMuted = Algo::AllOf(ActivitiesToDelete, [](const TSharedRef<FConcertSessionActivity>& Activity)
+	{
+		return (Activity->Activity.Flags & EConcertSyncActivityFlags::Muted) != EConcertSyncActivityFlags::None;
+	});
+	if (!bAreAllMuted)
+	{
+		return FCanPerformActionResult::No(LOCTEXT("Mute.CanMuteActivity.NotAllMuted", "Some of the selected activities are not muted."));
 	}
 
-	return FCanDeleteActivitiesResult::Yes();
+	return FCanPerformActionResult::Yes();
+}
+
+void FArchivedConcertSessionTab::SpawnDialogForMutingOrUnmuting(bool bShouldMute, const TSet<TSharedRef<FConcertSessionActivity>>& Activities) const
+{
+	using namespace UE::ConcertSyncCore;
+	using namespace UE::MultiUserServer;
+	
+	if (const TOptional<FConcertSyncSessionDatabaseNonNullPtr> SessionDatabase = SyncServer->GetArchivedSessionDatabase(InspectedSessionID))
+	{
+		TSet<FActivityID> Requested;
+		Algo::Transform(Activities, Requested, [](const TSharedRef<FConcertSessionActivity>& Activity)
+		{
+			return Activity->Activity.ActivityId;
+		});
+
+		const TWeakPtr<const FArchivedConcertSessionTab> WeakTabThis = SharedThis(this);
+		const TSharedRef<SMuteActivityDependenciesDialog> Dialog = SNew(SMuteActivityDependenciesDialog, InspectedSessionID, SyncServer, MoveTemp(Requested))
+			.MuteOperation(bShouldMute ? SMuteActivityDependenciesDialog::EMuteOperation::Mute : SMuteActivityDependenciesDialog::EMuteOperation::Unmute)
+			.OnConfirmMute_Lambda([WeakTabThis, bShouldMute, this](const TSet<FActivityID>& Selection)
+			{
+				// Because the dialog is non-modal, the user may have closed the program in the mean time
+				if (const TSharedPtr<const FArchivedConcertSessionTab> PinnedThis = WeakTabThis.Pin())
+				{
+					const bool bSuccess = PinnedThis->SyncServer->GetArchivedSessionDatabase(PinnedThis->InspectedSessionID)->SetActivities(Selection, [bShouldMute](FConcertSyncActivity& Activity)
+					{
+						if (bShouldMute)
+						{
+							Activity.Flags |= EConcertSyncActivityFlags::Muted;
+						}
+						else
+						{
+							Activity.Flags &= ~EConcertSyncActivityFlags::Muted;
+						}
+					});
+					if (bSuccess)
+					{
+						// Need to update model for history widget
+						HistoryController->ReloadActivities();
+					}
+					else
+					{
+						UE_LOG(LogConcert, Error, TEXT("Failed to mute activities from session %s"), *PinnedThis->InspectedSessionID.ToString(EGuidFormats::DigitsWithHyphens));
+						
+						const TSharedRef<SMessageDialog> ErrorDialog = SNew(SMessageDialog)
+							.Title(LOCTEXT("ErrorMutingSessions.Title", "Error muting sessions"))
+							.Message(LOCTEXT("ErrorMutingSessions.Description", "Error writing to database"))
+							.Buttons({
+								SMessageDialog::FButton(LOCTEXT("Ok", "Ok"))
+								.SetPrimary(true)
+							});
+						ErrorDialog->Show();
+					}
+				}
+			});
+		
+		FConcertServerUIModule::Get()
+			.GetModalWindowManager()
+			->ShowFakeModalWindow(Dialog);
+	}
 }
 
-#undef LOCTEXT_NAMESPACE
\ No newline at end of file
+#undef LOCTEXT_NAMESPACE
diff --git a/Engine/Plugins/Developer/Concert/ConcertApp/MultiUserServer/Source/MultiUserServer/Private/Widgets/SessionTabs/Archived/ArchivedConcertSessionTab.h b/Engine/Plugins/Developer/Concert/ConcertApp/MultiUserServer/Source/MultiUserServer/Private/Widgets/SessionTabs/Archived/ArchivedConcertSessionTab.h
--- a/Engine/Plugins/Developer/Concert/ConcertApp/MultiUserServer/Source/MultiUserServer/Private/Widgets/SessionTabs/Archived/ArchivedConcertSessionTab.h
+++ b/Engine/Plugins/Developer/Concert/ConcertApp/MultiUserServer/Source/MultiUserServer/Private/Widgets/SessionTabs/Archived/ArchivedConcertSessionTab.h
@@ -40,6 +40,14 @@ private:
 	/** Displays session */
 	TSharedPtr<SConcertArchivedSessionTabView> Inspector;
 
-	void OnRequestDeleteActivity(const TSet<TSharedRef<FConcertSessionActivity>>& ActivitiesToDelete) const;
-	FCanDeleteActivitiesResult CanDeleteActivity(const TSet<TSharedRef<FConcertSessionActivity>>& ActivitiesToDelete) const;
+	void OnRequestDeleteActivities(const TSet<TSharedRef<FConcertSessionActivity>>& ActivitiesToDelete) const;
+	FCanPerformActionResult CanDeleteActivities(const TSet<TSharedRef<FConcertSessionActivity>>& ActivitiesToDelete) const;
+	
+	void OnRequestMuteActivities(const TSet<TSharedRef<FConcertSessionActivity>>& ActivitiesToMute) const;
+	FCanPerformActionResult CanMuteActivities(const TSet<TSharedRef<FConcertSessionActivity>>& ActivitiesToDelete) const;
+	
+	void OnRequestUnmuteActivities(const TSet<TSharedRef<FConcertSessionActivity>>& ActivitiesToDelete) const;
+	FCanPerformActionResult CanUnmuteActivities(const TSet<TSharedRef<FConcertSessionActivity>>& ActivitiesToDelete) const;
+
+	void SpawnDialogForMutingOrUnmuting(bool bShouldMute, const TSet<TSharedRef<FConcertSessionActivity>>& Activities) const;
 };
diff --git a/Engine/Plugins/Developer/Concert/ConcertApp/MultiUserServer/Source/MultiUserServer/Private/Widgets/SessionTabs/Archived/SConcertArchivedSessionTabView.cpp b/Engine/Plugins/Developer/Concert/ConcertApp/MultiUserServer/Source/MultiUserServer/Private/Widgets/SessionTabs/Archived/SConcertArchivedSessionTabView.cpp
--- a/Engine/Plugins/Developer/Concert/ConcertApp/MultiUserServer/Source/MultiUserServer/Private/Widgets/SessionTabs/Archived/SConcertArchivedSessionTabView.cpp
+++ b/Engine/Plugins/Developer/Concert/ConcertApp/MultiUserServer/Source/MultiUserServer/Private/Widgets/SessionTabs/Archived/SConcertArchivedSessionTabView.cpp
@@ -12,7 +12,7 @@ const FName SConcertArchivedSessionTabView::HistoryTabId("HistoryTabId");
 
 void SConcertArchivedSessionTabView::Construct(const FArguments& InArgs, FName InStatusBarID)
 {
-	check(InArgs._MakeSessionHistory.IsBound() && InArgs._CanDeleteActivity.IsBound());
+	check(InArgs._MakeSessionHistory.IsBound() && InArgs._CanDeleteActivities.IsBound() && InArgs._CanMuteActivities.IsBound());
 	SConcertTabViewWithManagerBase::Construct(
 		SConcertTabViewWithManagerBase::FArguments()
 		.ConstructUnderWindow(InArgs._ConstructUnderWindow)
@@ -28,7 +28,15 @@ void SConcertArchivedSessionTabView::Construct(const FArguments& InArgs, FName I
 
 void SConcertArchivedSessionTabView::CreateTabs(const TSharedRef<FTabManager>& InTabManager, const TSharedRef<FTabManager::FLayout>& InLayout, const FArguments& InArgs)
 {
-	InTabManager->RegisterTabSpawner(HistoryTabId, FOnSpawnTab::CreateSP(this, &SConcertArchivedSessionTabView::SpawnActivityHistory, InArgs._MakeSessionHistory, InArgs._CanDeleteActivity, InArgs._DeleteActivity))
+	InTabManager->RegisterTabSpawner(HistoryTabId, FOnSpawnTab::CreateSP(this, &SConcertArchivedSessionTabView::SpawnActivityHistory,
+		InArgs._MakeSessionHistory,
+		InArgs._CanDeleteActivities,
+		InArgs._DeleteActivities,
+		InArgs._CanMuteActivities,
+		InArgs._MuteActivities,
+		InArgs._CanUnmuteActivities,
+		InArgs._UnmuteActivities
+		))
 		.SetDisplayName(LOCTEXT("ActivityHistoryLabel", "History"));
 	InLayout->AddArea
 		(
@@ -46,13 +54,21 @@ void SConcertArchivedSessionTabView::CreateTabs(const TSharedRef<FTabManager>& I
 TSharedRef<SDockTab> SConcertArchivedSessionTabView::SpawnActivityHistory(
 	const FSpawnTabArgs& Args,
 	SEditableSessionHistory::FMakeSessionHistory MakeSessionHistory,
-	SEditableSessionHistory::FCanDeleteActivities CanDeleteActivity,
-	SEditableSessionHistory::FRequestDeleteActivities DeleteActivity)
+	SEditableSessionHistory::FCanPerformActionOnActivities CanDeleteActivities,
+	SEditableSessionHistory::FRequestActivitiesAction DeleteActivities,
+	SEditableSessionHistory::FCanPerformActionOnActivities CanMuteActivities,
+	SEditableSessionHistory::FRequestActivitiesAction MuteActivities,
+	SEditableSessionHistory::FCanPerformActionOnActivities CanUnmuteActivities,
+	SEditableSessionHistory::FRequestActivitiesAction UnmuteActivities)
 {
 	SessionHistory = SNew(SEditableSessionHistory)
 		.MakeSessionHistory(MoveTemp(MakeSessionHistory))
-		.CanDeleteActivity(MoveTemp(CanDeleteActivity))
-		.DeleteActivity(MoveTemp(DeleteActivity));
+		.CanDeleteActivities(MoveTemp(CanDeleteActivities))
+		.DeleteActivities(MoveTemp(DeleteActivities))
+		.CanMuteActivities(MoveTemp(CanMuteActivities))
+		.MuteActivities(MoveTemp(MuteActivities))
+		.CanUnmuteActivities(MoveTemp(CanUnmuteActivities))
+		.UnmuteActivities(MoveTemp(UnmuteActivities));
 	return SNew(SDockTab)
 		.Label(LOCTEXT("ActivityHistoryLabel", "History"))
 		.TabRole(PanelTab)
diff --git a/Engine/Plugins/Developer/Concert/ConcertApp/MultiUserServer/Source/MultiUserServer/Private/Widgets/SessionTabs/Archived/SConcertArchivedSessionTabView.h b/Engine/Plugins/Developer/Concert/ConcertApp/MultiUserServer/Source/MultiUserServer/Private/Widgets/SessionTabs/Archived/SConcertArchivedSessionTabView.h
--- a/Engine/Plugins/Developer/Concert/ConcertApp/MultiUserServer/Source/MultiUserServer/Private/Widgets/SessionTabs/Archived/SConcertArchivedSessionTabView.h
+++ b/Engine/Plugins/Developer/Concert/ConcertApp/MultiUserServer/Source/MultiUserServer/Private/Widgets/SessionTabs/Archived/SConcertArchivedSessionTabView.h
@@ -25,8 +25,15 @@ public:
 		SLATE_ARGUMENT(TSharedPtr<SDockTab>, ConstructUnderMajorTab)
 		SLATE_ARGUMENT(TSharedPtr<SWindow>, ConstructUnderWindow)
 		SLATE_EVENT(SEditableSessionHistory::FMakeSessionHistory, MakeSessionHistory)
-		SLATE_EVENT(SEditableSessionHistory::FCanDeleteActivities, CanDeleteActivity)
-		SLATE_EVENT(SEditableSessionHistory::FRequestDeleteActivities, DeleteActivity)
+	
+		SLATE_EVENT(SEditableSessionHistory::FCanPerformActionOnActivities, CanDeleteActivities)
+		SLATE_EVENT(SEditableSessionHistory::FRequestActivitiesAction, DeleteActivities)
+	
+		SLATE_EVENT(SEditableSessionHistory::FCanPerformActionOnActivities, CanMuteActivities)
+		SLATE_EVENT(SEditableSessionHistory::FRequestActivitiesAction, MuteActivities)
+	
+		SLATE_EVENT(SEditableSessionHistory::FCanPerformActionOnActivities, CanUnmuteActivities)
+		SLATE_EVENT(SEditableSessionHistory::FRequestActivitiesAction, UnmuteActivities)
 	SLATE_END_ARGS()
 
 	void Construct(const FArguments& InArgs, FName InStatusBarID);
@@ -38,7 +45,11 @@ private:
 	void CreateTabs(const TSharedRef<FTabManager>& InTabManager, const TSharedRef<FTabManager::FLayout>& InLayout, const FArguments& InArgs);
 	TSharedRef<SDockTab> SpawnActivityHistory(const FSpawnTabArgs& Args,
 		SEditableSessionHistory::FMakeSessionHistory FMakeSessionHistory,
-		SEditableSessionHistory::FCanDeleteActivities CanDeleteActivity,
-		SEditableSessionHistory::FRequestDeleteActivities DeleteActivity
+		SEditableSessionHistory::FCanPerformActionOnActivities CanDeleteActivities,
+		SEditableSessionHistory::FRequestActivitiesAction DeleteActivities,
+		SEditableSessionHistory::FCanPerformActionOnActivities CanMuteActivities,
+		SEditableSessionHistory::FRequestActivitiesAction MuteActivities,
+		SEditableSessionHistory::FCanPerformActionOnActivities CanUnmuteActivities,
+		SEditableSessionHistory::FRequestActivitiesAction UnmuteActivities
 		);
 };
diff --git a/Engine/Plugins/Developer/Concert/ConcertUI/ConcertSharedSlate/Source/ConcertSharedSlate/Private/ConcertFrontendStyle.cpp b/Engine/Plugins/Developer/Concert/ConcertUI/ConcertSharedSlate/Source/ConcertSharedSlate/Private/ConcertFrontendStyle.cpp
--- a/Engine/Plugins/Developer/Concert/ConcertUI/ConcertSharedSlate/Source/ConcertSharedSlate/Private/ConcertFrontendStyle.cpp
+++ b/Engine/Plugins/Developer/Concert/ConcertUI/ConcertSharedSlate/Source/ConcertSharedSlate/Private/ConcertFrontendStyle.cpp
@@ -94,6 +94,10 @@ void FConcertFrontendStyle::Initialize()
 	StyleSet->Set("Concert.HidePresence",       new IMAGE_PLUGIN_BRUSH("Icons/icon_PresenceEyeOff_32x",     Icon16x16, IconColorAndOpacity));
 	StyleSet->Set("Concert.ShowPresence",       new IMAGE_PLUGIN_BRUSH("Icons/icon_PresenceEyeOn_32x",      Icon16x16, IconColorAndOpacity));
 
+	// Multi-user archived session browser
+	StyleSet->Set("Concert.MuteActivities",     new IMAGE_PLUGIN_BRUSH("Icons/icon_PresenceEyeOff_32x",      Icon16x16, IconColorAndOpacity));
+	StyleSet->Set("Concert.UnmuteActivities",     new IMAGE_PLUGIN_BRUSH("Icons/icon_PresenceEyeOn_32x",      Icon16x16, IconColorAndOpacity));
+
 	// 24x24/48x48 -> For sequencer toolbar.
 	StyleSet->Set("Concert.Sequencer.SyncTimeline",       new IMAGE_PLUGIN_BRUSH("Icons/icon_SequencerSyncTimeline_48x", Icon48x48, IconColorAndOpacity)); // Enable/disable playback and time scrubbing from a remote client.
 	StyleSet->Set("Concert.Sequencer.SyncTimeline.Small", new IMAGE_PLUGIN_BRUSH("Icons/icon_SequencerSyncTimeline_48x", Icon24x24, IconColorAndOpacity));
diff --git a/Engine/Plugins/Developer/Concert/ConcertUI/ConcertSharedSlate/Source/ConcertSharedSlate/Private/Session/Activity/SConcertSessionActivities.cpp b/Engine/Plugins/Developer/Concert/ConcertUI/ConcertSharedSlate/Source/ConcertSharedSlate/Private/Session/Activity/SConcertSessionActivities.cpp
--- a/Engine/Plugins/Developer/Concert/ConcertUI/ConcertSharedSlate/Source/ConcertSharedSlate/Private/Session/Activity/SConcertSessionActivities.cpp
+++ b/Engine/Plugins/Developer/Concert/ConcertUI/ConcertSharedSlate/Source/ConcertSharedSlate/Private/Session/Activity/SConcertSessionActivities.cpp
@@ -79,9 +79,12 @@ public:
 	
 	SLATE_BEGIN_ARGS(SConcertSessionActivityRow)
 		: _OnMakeColumnOverlayWidget()
+		, _DarkenMutedActivities(true)
 	{}
 		/** Function invoked when generating a row to add a widget above the column widget. */
-		SLATE_ARGUMENT(SConcertSessionActivities::FMakeColumnOverlayWidgetFunc, OnMakeColumnOverlayWidget) 
+		SLATE_ARGUMENT(SConcertSessionActivities::FMakeColumnOverlayWidgetFunc, OnMakeColumnOverlayWidget)
+		/** Optional. Whether to reduce focus to activities by darkening when the activity is muted (default: true). */
+		SLATE_ARGUMENT(bool, DarkenMutedActivities) 
 	SLATE_END_ARGS()
 
 	/**
@@ -111,6 +114,7 @@ private:
 	FText AbsoluteDateTime;
 	FText ClientName;
 	SConcertSessionActivities::FMakeColumnOverlayWidgetFunc OnMakeColumnOverlayWidget;
+	bool bDarkenMutedActivities;
 };
 
 
@@ -124,11 +128,13 @@ void SConcertSessionActivityRow::Construct(const FArguments& InArgs, TSharedRef<
 	OnMakeColumnOverlayWidget = InArgs._OnMakeColumnOverlayWidget;
 	AbsoluteDateTime = ConcertFrontendUtils::FormatTime(InActivity->Activity.EventTime, ETimeFormat::Absolute); // Cache the absolute time. It doesn't changes.
 	ClientName = ConcertSessionActivityUtils::GetClientName(InActivityClient);
-
+	bDarkenMutedActivities = InArgs._DarkenMutedActivities;
+	
 	// Construct base class
 	SMultiColumnTableRow<TSharedPtr<FConcertSessionActivity>>::Construct(FSuperRowType::FArguments(), InOwnerTableView);
 
-	if (InActivity->Activity.bIgnored)
+	const bool bDarkenMutedActivity = bDarkenMutedActivities && (InActivity->Activity.Flags & EConcertSyncActivityFlags::Muted) != EConcertSyncActivityFlags::None;
+	if (InActivity->Activity.bIgnored || bDarkenMutedActivity)
 	{
 		SetColorAndOpacity(FLinearColor(0.5, 0.5, 0.5, 0.5));
 	}
@@ -195,6 +201,12 @@ FText SConcertSessionActivityRow::MakeTooltipText() const
 			TextBuilder.AppendLine();
 			TextBuilder.AppendLine(LOCTEXT("IgnoredActivity", "** This activity cannot be recovered (likely recorded during a Multi-User session). It is displayed for crash inspection only. It will be ignored on restore."));
 		}
+		
+		const bool bDarkenMutedActivity = bDarkenMutedActivities && (ActivityPin->Activity.Flags & EConcertSyncActivityFlags::Muted) != EConcertSyncActivityFlags::None;
+		if (bDarkenMutedActivity)
+		{
+			TextBuilder.AppendLine(LOCTEXT("MutedActivity", "** This activity is muted."));
+		}
 
 		return TextBuilder.ToText();
 	}
@@ -209,6 +221,7 @@ void SConcertSessionActivities::Construct(const FArguments& InArgs)
 	GetTransactionEventFn = InArgs._OnGetTransactionEvent;
 	GetPackageEventFn = InArgs._OnGetPackageEvent;
 	MakeColumnOverlayWidgetFn = InArgs._OnMakeColumnOverlayWidget;
+
 	
 	HighlightText = InArgs._HighlightText;
 	
@@ -220,6 +233,7 @@ void SConcertSessionActivities::Construct(const FArguments& InArgs)
 	IgnoredActivitiesVisibility = InArgs._IgnoredActivitiesVisibility;
 	DetailsAreaVisibility = InArgs._DetailsAreaVisibility;
 	bAutoScrollDesired = InArgs._IsAutoScrollEnabled;
+	bDarkenMutedActivities = InArgs._DarkenMutedActivities;
 
 	SearchTextFilter = MakeShared<TTextFilter<const FConcertSessionActivity&>>(TTextFilter<const FConcertSessionActivity&>::FItemToStringArray::CreateSP(this, &SConcertSessionActivities::PopulateSearchStrings));
 	SearchTextFilter->OnChanged().AddSP(this, &SConcertSessionActivities::OnActivityFilterUpdated);
@@ -410,7 +424,8 @@ TSharedRef<ITableRow> SConcertSessionActivities::OnGenerateActivityRowWidget(TSh
 	});
 	
 	return SNew(SConcertSessionActivityRow, SharedThis(this), Activity.ToSharedRef(), ColumnGetter, ActivityClient, OwnerTable)
-		.OnMakeColumnOverlayWidget(MakeColumnOverlayWidgetFn);
+		.OnMakeColumnOverlayWidget(MakeColumnOverlayWidgetFn)
+		.DarkenMutedActivities(bDarkenMutedActivities);
 }
 
 void SConcertSessionActivities::Tick(const FGeometry& AllottedGeometry, const double InCurrentTime, const float InDeltaTime)
diff --git a/Engine/Plugins/Developer/Concert/ConcertUI/ConcertSharedSlate/Source/ConcertSharedSlate/Private/Session/History/SEditableSessionHistory.cpp b/Engine/Plugins/Developer/Concert/ConcertUI/ConcertSharedSlate/Source/ConcertSharedSlate/Private/Session/History/SEditableSessionHistory.cpp
--- a/Engine/Plugins/Developer/Concert/ConcertUI/ConcertSharedSlate/Source/ConcertSharedSlate/Private/Session/History/SEditableSessionHistory.cpp
+++ b/Engine/Plugins/Developer/Concert/ConcertUI/ConcertSharedSlate/Source/ConcertSharedSlate/Private/Session/History/SEditableSessionHistory.cpp
@@ -2,32 +2,71 @@
 
 #include "Session/History/SEditableSessionHistory.h"
 
+#include "Algo/Transform.h"
 #include "Styling/AppStyle.h"
 #include "SNegativeActionButton.h"
-
-#include "Algo/Transform.h"
-
 #include "Session/History/SSessionHistory.h"
-#include "Widgets/Layout/SBox.h"
+#include "Widgets/SBoxPanel.h"
 
 #define LOCTEXT_NAMESPACE "SEditableSessionHistory"
 
 void SEditableSessionHistory::Construct(const FArguments& InArgs)
 {
-	check(InArgs._MakeSessionHistory.IsBound() && InArgs._CanDeleteActivity.IsBound());
-	CanDeleteActivityFunc = InArgs._CanDeleteActivity;
-	DeleteActivityFunc = InArgs._DeleteActivity;
+	check(InArgs._MakeSessionHistory.IsBound()
+		&& InArgs._CanDeleteActivities.IsBound()
+		&& InArgs._CanMuteActivities.IsBound()
+		&& InArgs._CanUnmuteActivities.IsBound()
+		);
+	
+	CanDeleteActivitiesFunc = InArgs._CanDeleteActivities;
+	DeleteActivitiesFunc = InArgs._DeleteActivities;
+	CanMuteActivitiesFunc = InArgs._CanMuteActivities;
+	MuteActivitiesFunc = InArgs._MuteActivities;
+	CanUnmuteActivitiesFunc = InArgs._CanUnmuteActivities;
+	UnmuteActivitiesFunc = InArgs._UnmuteActivities;
 	
 	SessionHistory = InArgs._MakeSessionHistory.Execute(
 		SSessionHistory::FArguments()
 		.SelectionMode(ESelectionMode::Multi)
 		.SearchButtonArea()
 		[
-			SNew(SNegativeActionButton)
-			.OnClicked(this, &SEditableSessionHistory::OnClickDeleteActivityButton)
-			.ToolTipText(this, &SEditableSessionHistory::GetDeleteActivityToolTip)
-			.IsEnabled(this, &SEditableSessionHistory::IsDeleteButtonEnabled)
-			.Icon(FAppStyle::GetBrush("Icons.Delete"))
+			SNew(SHorizontalBox)
+
+			// Delete
+			+SHorizontalBox::Slot()
+			.AutoWidth()
+			[
+				SNew(SNegativeActionButton)
+				.ActionButtonStyle(EActionButtonStyle::Error) // Red
+				.OnClicked(this, &SEditableSessionHistory::OnClickDeleteActivitiesButton)
+				.ToolTipText(this, &SEditableSessionHistory::GetDeleteActivitiesToolTip)
+				.IsEnabled(this, &SEditableSessionHistory::IsDeleteButtonEnabled)
+				.Icon(FAppStyle::GetBrush("Icons.Delete"))
+			]
+
+			// Mute
+			+SHorizontalBox::Slot()
+			.AutoWidth()
+			[
+				SNew(SNegativeActionButton)
+				.ActionButtonStyle(EActionButtonStyle::Warning) // Yellow
+				.OnClicked(this, &SEditableSessionHistory::OnClickMuteActivitesButton)
+				.ToolTipText(this, &SEditableSessionHistory::GetMuteActivitiesToolTip)
+				.IsEnabled(this, &SEditableSessionHistory::IsMuteButtonEnabled)
+				.Icon(FConcertFrontendStyle::Get()->GetBrush("Concert.MuteActivities"))
+			]
+
+			// Unmute
+			+SHorizontalBox::Slot()
+			.AutoWidth()
+			[
+				SNew(SNegativeActionButton)
+				.ActionButtonStyle(EActionButtonStyle::Warning) // Yellow
+				.OnClicked(this, &SEditableSessionHistory::OnClickUnmuteActivitesButton)
+				.ToolTipText(this, &SEditableSessionHistory::GetUnmuteActivitiesToolTip)
+				.IsEnabled(this, &SEditableSessionHistory::IsUnmuteButtonEnabled)
+				.Icon(FConcertFrontendStyle::Get()->GetBrush("Concert.UnmuteActivities"))
+			]
 		]
 		);
 	
@@ -42,9 +81,9 @@ FReply SEditableSessionHistory::OnKeyDown(const FGeometry& MyGeometry, const FKe
 	if (InKeyEvent.GetKey() == EKeys::Delete)
 	{
 		const TSet<TSharedRef<FConcertSessionActivity>> SelectedActivities = SessionHistory->GetSelectedActivities();
-		if (CanDeleteActivityFunc.Execute(SelectedActivities).CanDelete())
+		if (CanDeleteActivitiesFunc.Execute(SelectedActivities).CanPerformAction())
 		{
-			DeleteActivityFunc.Execute(SelectedActivities);
+			DeleteActivitiesFunc.Execute(SelectedActivities);
 		}
 		return FReply::Handled();
 	}
@@ -52,49 +91,104 @@ FReply SEditableSessionHistory::OnKeyDown(const FGeometry& MyGeometry, const FKe
 	return SCompoundWidget::OnKeyDown(MyGeometry, InKeyEvent);
 }
 
-FReply SEditableSessionHistory::OnClickDeleteActivityButton() const
+FReply SEditableSessionHistory::OnClickDeleteActivitiesButton() const
+{
+	if (DeleteActivitiesFunc.IsBound() && CanDeleteActivitiesFunc.Execute(SessionHistory->GetSelectedActivities()).CanPerformAction())
+	{
+		DeleteActivitiesFunc.Execute(SessionHistory->GetSelectedActivities());
+	}
+	return FReply::Handled();
+}
+
+FText SEditableSessionHistory::GetDeleteActivitiesToolTip() const
+{
+	return GenerateTooltip(
+		CanDeleteActivitiesFunc,
+		LOCTEXT("Delete.SelectActivityToolTip", "Select some activities to delete from below (multi-select using CTRL + Click)."),
+		LOCTEXT("Delete.SelectedActivitiesToolTip", "Delete selection from history (IDs: {0})"),
+		LOCTEXT("Delete.CannotSelectedActivitiesToolTip", "Cannot delete:\n{0}")
+		);
+}
+
+bool SEditableSessionHistory::IsDeleteButtonEnabled() const
+{
+	const TSet<TSharedRef<FConcertSessionActivity>> SelectedActivities = SessionHistory->GetSelectedActivities();
+	return SelectedActivities.Num() > 0 && CanDeleteActivitiesFunc.Execute(SelectedActivities).CanPerformAction();
+}
+
+FReply SEditableSessionHistory::OnClickMuteActivitesButton() const
+{
+	if (MuteActivitiesFunc.IsBound() && CanMuteActivitiesFunc.Execute(SessionHistory->GetSelectedActivities()).CanPerformAction())
+	{
+		MuteActivitiesFunc.Execute(SessionHistory->GetSelectedActivities());
+	}
+	return FReply::Handled();
+}
+
+FText SEditableSessionHistory::GetMuteActivitiesToolTip() const
+{
+	return GenerateTooltip(
+		CanMuteActivitiesFunc,
+		LOCTEXT("Mute.SelectActivityToolTip", "Select some activities to mute (multi-select using CTRL + Click)."),
+		LOCTEXT("Mute.SelectedActivitiesToolTip", "Mute selection ({0})"),
+		LOCTEXT("Mute.CannotSelectedActivitiesToolTip", "Cannot mute:\n{0}")
+		);
+}
+
+bool SEditableSessionHistory::IsMuteButtonEnabled() const
+{
+	const TSet<TSharedRef<FConcertSessionActivity>> SelectedActivities = SessionHistory->GetSelectedActivities();
+	return SelectedActivities.Num() > 0 && CanMuteActivitiesFunc.Execute(SelectedActivities).CanPerformAction();
+}
+
+FReply SEditableSessionHistory::OnClickUnmuteActivitesButton() const
 {
-	if (DeleteActivityFunc.IsBound() && CanDeleteActivityFunc.Execute(SessionHistory->GetSelectedActivities()).CanDelete())
+	if (UnmuteActivitiesFunc.IsBound() && CanUnmuteActivitiesFunc.Execute(SessionHistory->GetSelectedActivities()).CanPerformAction())
 	{
-		DeleteActivityFunc.Execute(SessionHistory->GetSelectedActivities());
+		UnmuteActivitiesFunc.Execute(SessionHistory->GetSelectedActivities());
 	}
 	return FReply::Handled();
 }
 
-FText SEditableSessionHistory::GetDeleteActivityToolTip() const
+FText SEditableSessionHistory::GetUnmuteActivitiesToolTip() const
+{
+	return GenerateTooltip(
+		CanUnmuteActivitiesFunc,
+		LOCTEXT("Unmute.SelectActivityToolTip", "Select some activities to unmute (multi-select using CTRL + Click)."),
+		LOCTEXT("Unmute.SelectedActivitiesToolTip", "Unmute selection ({0})"),
+		LOCTEXT("Unmute.CannotSelectedActivitiesToolTip", "Cannot unmute:\n{0}")
+		);
+}
+
+bool SEditableSessionHistory::IsUnmuteButtonEnabled() const
+{
+	const TSet<TSharedRef<FConcertSessionActivity>> SelectedActivities = SessionHistory->GetSelectedActivities();
+	return SelectedActivities.Num() > 0 && CanUnmuteActivitiesFunc.Execute(SelectedActivities).CanPerformAction();
+}
+
+FText SEditableSessionHistory::GenerateTooltip(const FCanPerformActionOnActivities& CanPerformAction, FText SelectActivityToolTip, FText PerformActionToolTipFmt, FText CannotPerformActionToolTipFmt) const
 {
 	TSet<TSharedRef<FConcertSessionActivity>> SelectedActivities = SessionHistory->GetSelectedActivities();
 	if (SelectedActivities.Num() == 0)
 	{
-		return LOCTEXT("SelectActivityToolTip", "Select some activities to delete from below (multi-select using CTRL + Click).");
+		return SelectActivityToolTip;
 	}
 	
-	const FCanDeleteActivitiesResult CanDeleteActivities = CanDeleteActivityFunc.Execute(SelectedActivities);
-	if (CanDeleteActivities.CanDelete())
+	const FCanPerformActionResult CanDeleteActivities = CanPerformAction.Execute(SelectedActivities);
+	if (CanDeleteActivities.CanPerformAction())
 	{
 		SelectedActivities.Sort([](const TSharedRef<FConcertSessionActivity>& First, const TSharedRef<FConcertSessionActivity>& Second) { return First->Activity.ActivityId <= Second->Activity.ActivityId; });
-		
-		TStringBuilder<256> ActivityBuilder;
-		bool bNeedsComma = false;
-		for (const TSharedRef<FConcertSessionActivity>& Activity : SelectedActivities)
-		{
-			if (bNeedsComma)
-			{
-				ActivityBuilder << ", ";
-			}
-			ActivityBuilder << Activity->Activity.ActivityId;
 
-			bNeedsComma = true;
-		}
-		return FText::Format(LOCTEXT("DeleteSelectedActivitiesToolTip", "Delete selected activities from history (IDs: {0})"), FText::FromString(ActivityBuilder.ToString()));
+		TArray<FString> Items;
+		Algo::Transform(SelectedActivities, Items, [](const TSharedRef<FConcertSessionActivity>& Activity)
+		{
+			return FString::FromInt(Activity->Activity.ActivityId);
+		});
+		const FString AllItems = FString::Join(Items, TEXT(", "));
+		
+		return FText::Format(PerformActionToolTipFmt, FText::FromString(AllItems));
 	}
-	return FText::Format(LOCTEXT("CannotDeleteSelectedActivitiesToolTip", "Activity cannot be deleted: {0}"), CanDeleteActivities.DeletionReason.GetValue());
-}
-
-bool SEditableSessionHistory::IsDeleteButtonEnabled() const
-{
-	const TSet<TSharedRef<FConcertSessionActivity>> SelectedActivities = SessionHistory->GetSelectedActivities();
-	return SelectedActivities.Num() > 0 && CanDeleteActivityFunc.Execute(SelectedActivities).CanDelete();
+	return FText::Format(CannotPerformActionToolTipFmt, CanDeleteActivities.DeletionReason.GetValue());
 }
 
 #undef LOCTEXT_NAMESPACE
diff --git a/Engine/Plugins/Developer/Concert/ConcertUI/ConcertSharedSlate/Source/ConcertSharedSlate/Private/Session/History/SSessionHistory.cpp b/Engine/Plugins/Developer/Concert/ConcertUI/ConcertSharedSlate/Source/ConcertSharedSlate/Private/Session/History/SSessionHistory.cpp
--- a/Engine/Plugins/Developer/Concert/ConcertUI/ConcertSharedSlate/Source/ConcertSharedSlate/Private/Session/History/SSessionHistory.cpp
+++ b/Engine/Plugins/Developer/Concert/ConcertUI/ConcertSharedSlate/Source/ConcertSharedSlate/Private/Session/History/SSessionHistory.cpp
@@ -73,7 +73,8 @@ void SSessionHistory::Construct(const FArguments& InArgs)
 		.IsAutoScrollEnabled(true)
 		.ColumnVisibilitySnapshot(InArgs._ColumnVisibilitySnapshot)
 		.SaveColumnVisibilitySnapshot(InArgs._SaveColumnVisibilitySnapshot)
-		.SelectionMode(InArgs._SelectionMode);
+		.SelectionMode(InArgs._SelectionMode)
+		.DarkenMutedActivities(InArgs._DarkenMutedActivities);
 
 	ChildSlot
 	[
diff --git a/Engine/Plugins/Developer/Concert/ConcertUI/ConcertSharedSlate/Source/ConcertSharedSlate/Public/Session/Activity/SConcertSessionActivities.h b/Engine/Plugins/Developer/Concert/ConcertUI/ConcertSharedSlate/Source/ConcertSharedSlate/Public/Session/Activity/SConcertSessionActivities.h
--- a/Engine/Plugins/Developer/Concert/ConcertUI/ConcertSharedSlate/Source/ConcertSharedSlate/Public/Session/Activity/SConcertSessionActivities.h
+++ b/Engine/Plugins/Developer/Concert/ConcertUI/ConcertSharedSlate/Source/ConcertSharedSlate/Public/Session/Activity/SConcertSessionActivities.h
@@ -84,7 +84,9 @@ public:
 		, _TransactionActivitiesVisibility(EVisibility::Visible)
 		, _IgnoredActivitiesVisibility(EVisibility::Hidden)
 		, _DetailsAreaVisibility(EVisibility::Hidden)
-		, _IsAutoScrollEnabled(false){ }
+		, _IsAutoScrollEnabled(false)
+		, _DarkenMutedActivities(true)
+		{ }
 
 		/** If bound, invoked to populate the view. */
 		SLATE_EVENT(FFetchActivitiesFunc, OnFetchActivities)
@@ -140,8 +142,11 @@ public:
 		SLATE_EVENT(UE::ConcertSharedSlate::FSaveColumnVisibilitySnapshot, SaveColumnVisibilitySnapshot)
 	
 
-		/** How the activities may be selected */
+		/** Optional. How the activities may be selected */
 		SLATE_ARGUMENT(TOptional<ESelectionMode::Type>, SelectionMode)
+
+		/** Optional. Whether to reduce focus to activities by darkening when the activity is muted (default: true). */
+		SLATE_ARGUMENT(bool, DarkenMutedActivities)
 	SLATE_END_ARGS();
 	
 	/**
@@ -192,6 +197,13 @@ public:
 	TOptional<FConcertClientInfo> GetClientInfo(const FGuid& Guid) const { return GetActivityUserFn.IsBound() ? GetActivityUserFn.Execute(Guid) : TOptional<FConcertClientInfo>{}; }
 	TSharedPtr<SHeaderRow> GetHeaderRow() const { return HeaderRow; }
 
+	const TArray<TSharedPtr<FConcertSessionActivity>>& GetActivities() const { return Activities; }
+	void SetSelectedActivities(const TArray<TSharedPtr<FConcertSessionActivity>>& ActivitiesToSelect)
+	{
+		ActivityView->ClearSelection();
+		ActivityView->SetItemSelection(ActivitiesToSelect, true);
+	}
+
 private:
 
 	// Widget creation
@@ -238,6 +250,9 @@ private:
 	/** Used to overlay a widget over a column widget (add an extra layer above the normal one) */
 	FMakeColumnOverlayWidgetFunc MakeColumnOverlayWidgetFn;
 
+	/** Whether to darken muted activities */
+	bool bDarkenMutedActivities = true;
+
 	/** Returns which text should be highlighted. */
 	TAttribute<FText> HighlightText;
 
diff --git a/Engine/Plugins/Developer/Concert/ConcertUI/ConcertSharedSlate/Source/ConcertSharedSlate/Public/Session/History/SEditableSessionHistory.h b/Engine/Plugins/Developer/Concert/ConcertUI/ConcertSharedSlate/Source/ConcertSharedSlate/Public/Session/History/SEditableSessionHistory.h
--- a/Engine/Plugins/Developer/Concert/ConcertUI/ConcertSharedSlate/Source/ConcertSharedSlate/Public/Session/History/SEditableSessionHistory.h
+++ b/Engine/Plugins/Developer/Concert/ConcertUI/ConcertSharedSlate/Source/ConcertSharedSlate/Public/Session/History/SEditableSessionHistory.h
@@ -10,16 +10,16 @@
 class SSessionHistory;
 struct FConcertSessionActivity;
 
-struct FCanDeleteActivitiesResult
+struct FCanPerformActionResult
 {
 	TOptional<FText> DeletionReason;
 
-	static FCanDeleteActivitiesResult Yes() { return FCanDeleteActivitiesResult(); }
-	static FCanDeleteActivitiesResult No(FText Reason) { return FCanDeleteActivitiesResult(MoveTemp(Reason)); }
+	static FCanPerformActionResult Yes() { return FCanPerformActionResult(); }
+	static FCanPerformActionResult No(FText Reason) { return FCanPerformActionResult(MoveTemp(Reason)); }
 
-	bool CanDelete() const { return !DeletionReason.IsSet(); }
+	bool CanPerformAction() const { return !DeletionReason.IsSet(); }
 
-	explicit FCanDeleteActivitiesResult(TOptional<FText> DeletionReason = {})
+	explicit FCanPerformActionResult(TOptional<FText> DeletionReason = {})
 		: DeletionReason(DeletionReason)
 	{}
 };
@@ -30,14 +30,26 @@ class CONCERTSHAREDSLATE_API SEditableSessionHistory : public SCompoundWidget
 public:
 
 	DECLARE_DELEGATE_RetVal_OneParam(TSharedRef<SSessionHistory>, FMakeSessionHistory, SSessionHistory::FArguments)
-	DECLARE_DELEGATE_RetVal_OneParam(FCanDeleteActivitiesResult, FCanDeleteActivities, const TSet<TSharedRef<FConcertSessionActivity>>& /*Activities*/)
-	DECLARE_DELEGATE_OneParam(FRequestDeleteActivities, const TSet<TSharedRef<FConcertSessionActivity>>&)
+	DECLARE_DELEGATE_RetVal_OneParam(FCanPerformActionResult, FCanPerformActionOnActivities, const TSet<TSharedRef<FConcertSessionActivity>>& /*Activities*/)
+	DECLARE_DELEGATE_OneParam(FRequestActivitiesAction, const TSet<TSharedRef<FConcertSessionActivity>>&)
 
 	SLATE_BEGIN_ARGS(SEditableSessionHistory)
 	{}
 		SLATE_EVENT(FMakeSessionHistory, MakeSessionHistory)
-		SLATE_EVENT(FCanDeleteActivities, CanDeleteActivity)
-		SLATE_EVENT(FRequestDeleteActivities, DeleteActivity)
+	
+		/** Can selected activities be deleted? */
+		SLATE_EVENT(FCanPerformActionOnActivities, CanDeleteActivities)
+		/** Delete the selected activities */
+		SLATE_EVENT(FRequestActivitiesAction, DeleteActivities)
+		/** Can selected activities be muted? */
+		SLATE_EVENT(FCanPerformActionOnActivities, CanMuteActivities)
+		/** Mute the selected activities */
+		SLATE_EVENT(FRequestActivitiesAction, MuteActivities)
+		/** Can selected activities be unmuted? */
+		SLATE_EVENT(FCanPerformActionOnActivities, CanUnmuteActivities)
+		/** Unmute the selected activities */
+		SLATE_EVENT(FRequestActivitiesAction, UnmuteActivities)
+	
 		SLATE_NAMED_SLOT(FArguments, StatusBar)
 	SLATE_END_ARGS()
 
@@ -51,10 +63,24 @@ private:
 
 	TSharedPtr<SSessionHistory> SessionHistory;
 
-	FCanDeleteActivities CanDeleteActivityFunc;
-	FRequestDeleteActivities DeleteActivityFunc;
+	FCanPerformActionOnActivities CanDeleteActivitiesFunc;
+	FRequestActivitiesAction DeleteActivitiesFunc;
+	FCanPerformActionOnActivities CanMuteActivitiesFunc;
+	FRequestActivitiesAction MuteActivitiesFunc;
+	FCanPerformActionOnActivities CanUnmuteActivitiesFunc;
+	FRequestActivitiesAction UnmuteActivitiesFunc;
 
-	FReply OnClickDeleteActivityButton() const;
-	FText GetDeleteActivityToolTip() const;
+	FReply OnClickDeleteActivitiesButton() const;
+	FText GetDeleteActivitiesToolTip() const;
 	bool IsDeleteButtonEnabled() const;
+	
+	FReply OnClickMuteActivitesButton() const;
+	FText GetMuteActivitiesToolTip() const;
+	bool IsMuteButtonEnabled() const;
+	
+	FReply OnClickUnmuteActivitesButton() const;
+	FText GetUnmuteActivitiesToolTip() const;
+	bool IsUnmuteButtonEnabled() const;
+
+	FText GenerateTooltip(const FCanPerformActionOnActivities& CanPerformAction, FText SelectActivityToolTip, FText PerformActionToolTipFmt, FText CannotPerformActionToolTipFmt) const;
 };
diff --git a/Engine/Plugins/Developer/Concert/ConcertUI/ConcertSharedSlate/Source/ConcertSharedSlate/Public/Session/History/SSessionHistory.h b/Engine/Plugins/Developer/Concert/ConcertUI/ConcertSharedSlate/Source/ConcertSharedSlate/Public/Session/History/SSessionHistory.h
--- a/Engine/Plugins/Developer/Concert/ConcertUI/ConcertSharedSlate/Source/ConcertSharedSlate/Public/Session/History/SSessionHistory.h
+++ b/Engine/Plugins/Developer/Concert/ConcertUI/ConcertSharedSlate/Source/ConcertSharedSlate/Public/Session/History/SSessionHistory.h
@@ -29,6 +29,7 @@ public:
 		})
 		, _DetailsAreaVisibility(EVisibility::Visible)
 		, _SelectionMode(ESelectionMode::Single)
+		, _DarkenMutedActivities(true)
 	{}
 	
 		SLATE_EVENT(SConcertSessionActivities::FGetPackageEvent, GetPackageEvent)
@@ -58,6 +59,9 @@ public:
 		/** Optional. How the activities may be selected. Default to Single. */
 		SLATE_ARGUMENT(ESelectionMode::Type, SelectionMode)
 
+		/** Optional. Whether to reduce focus to activities by darkening when the activity is muted (default: true). */
+		SLATE_ARGUMENT(bool, DarkenMutedActivities)
+	
 		/** Optional. An area to the left of the search bar intended for adding buttons to. */
 		SLATE_NAMED_SLOT(FArguments, SearchButtonArea)
 	
@@ -79,6 +83,8 @@ public:
 	bool IsLastColumn(FName ColumnId) const { return ActivityListView->IsLastColumn(ColumnId); }
 	
 	TSet<TSharedRef<FConcertSessionActivity>> GetSelectedActivities() const;
+	const TArray<TSharedPtr<FConcertSessionActivity>>& GetActivities() const { return ActivityListView->GetActivities(); }
+	void SetSelectedActivities(const TArray<TSharedPtr<FConcertSessionActivity>>& ActivitiesToSelect) { ActivityListView->SetSelectedActivities(ActivitiesToSelect); }
 
 	void OnColumnVisibilitySettingsChanged(const FColumnVisibilitySnapshot& ColumnSnapshot);
 
