Fix delete, mute & unmute dialogs closing after pressing button

diff --git a/Engine/Plugins/Developer/Concert/ConcertApp/MultiUserServer/Source/MultiUserServer/Private/Widgets/HistoryEdition/SActivityDependencyDialog.cpp b/Engine/Plugins/Developer/Concert/ConcertApp/MultiUserServer/Source/MultiUserServer/Private/Widgets/HistoryEdition/SActivityDependencyDialog.cpp
--- a/Engine/Plugins/Developer/Concert/ConcertApp/MultiUserServer/Source/MultiUserServer/Private/Widgets/HistoryEdition/SActivityDependencyDialog.cpp
+++ b/Engine/Plugins/Developer/Concert/ConcertApp/MultiUserServer/Source/MultiUserServer/Private/Widgets/HistoryEdition/SActivityDependencyDialog.cpp
@@ -30,6 +30,7 @@ namespace UE::MultiUserServer
 		
 		SCustomDialog::Construct(
 			SCustomDialog::FArguments()
+			.AutoCloseOnButtonPress(false)
 			.Title(InArgs._Title)
 			.Buttons({
 				FButton(InArgs._PerformActionButtonLabel)
