Add graph to client thumbnail

diff --git a/Engine/Plugins/Developer/Concert/ConcertApp/MultiUserServer/Source/MultiUserServer/Private/ConcertServerStyle.cpp b/Engine/Plugins/Developer/Concert/ConcertApp/MultiUserServer/Source/MultiUserServer/Private/ConcertServerStyle.cpp
--- a/Engine/Plugins/Developer/Concert/ConcertApp/MultiUserServer/Source/MultiUserServer/Private/ConcertServerStyle.cpp
+++ b/Engine/Plugins/Developer/Concert/ConcertApp/MultiUserServer/Source/MultiUserServer/Private/ConcertServerStyle.cpp
@@ -84,6 +84,7 @@ TSharedRef<FSlateStyleSet> FConcertServerStyle::Create()
 	StyleSet->Set("Concert.Clients.ThumbnailAreaHoverBackground", new FSlateRoundedBoxBrush(FStyleColors::Hover, 4.0f));
 	StyleSet->Set("Concert.Clients.ThumbnailAreaBackground", new FSlateRoundedBoxBrush(FStyleColors::Secondary, 4.0f));
 	StyleSet->Set("Concert.Clients.ThumbnailTitle", new FSlateRoundedBoxBrush(FStyleColors::Recessed, 4.0f));
+	StyleSet->Set("Concert.Clients.ThumbnailCurveBackground", new FSlateRoundedBoxBrush(FStyleColors::Secondary, 0.0f));
 	StyleSet->Set("Concert.Clients.ThumbnailFooter", new FSlateRoundedBoxBrush(FStyleColors::Panel, 0.0f));
 	StyleSet->Set("Concert.Clients.TileTableRow", FTableRowStyle()
 			.SetEvenRowBackgroundBrush(FSlateNoResource() )
@@ -100,6 +101,22 @@ TSharedRef<FSlateStyleSet> FConcertServerStyle::Create()
 			);
 	StyleSet->Set("Concert.Clients.ClientNameTileFont", DEFAULT_FONT("Regular", 16));
 	
+	// Graph colours
+	// Blueish
+	StyleSet->Set("Concert.Clients.NetworkGraph.Sent.LineColor", 0.7f * FLinearColor(.24f, .38f, .76f, 0.f));
+	StyleSet->Set("Concert.Clients.NetworkGraph.Sent.FillColor",  0.4f * FLinearColor(.24f, .38f, .76f, 1.f));
+	// Reddish
+	StyleSet->Set("Concert.Clients.NetworkGraph.Received.LineColor", 0.7f * FLinearColor(.85f, .27f, .27f, 0.f));
+	StyleSet->Set("Concert.Clients.NetworkGraph.Received.FillColor",  0.4f * FLinearColor(.85f, .27f, .27f, 1.f));
+
+	constexpr float LineBrightness = 0.35f;
+	constexpr float TextBrightness = 0.5f;
+	StyleSet->Set("Concert.Clients.NetworkGraph.GraphSeparatorLine.Thickness", 3.f);
+	StyleSet->Set("Concert.Clients.NetworkGraph.GraphSeparatorLine.LineColor", FLinearColor(LineBrightness, LineBrightness, LineBrightness, 0.35f));
+	StyleSet->Set("Concert.Clients.NetworkGraph.HorizontalHelperLine.Thickness", 2.f);
+	StyleSet->Set("Concert.Clients.NetworkGraph.HorizontalHelperLine.LineColor", FLinearColor(LineBrightness, LineBrightness, LineBrightness, 0.35f));
+	StyleSet->Set("Concert.Clients.NetworkGraph.HorizontalHelperLine.TextColor", FLinearColor(TextBrightness, TextBrightness, TextBrightness));
+	
 	return StyleSet;
 }
 
diff --git a/Engine/Plugins/Developer/Concert/ConcertApp/MultiUserServer/Source/MultiUserServer/Private/Widgets/Clients/Browser/Graph/SClientNetworkGraphs.cpp b/Engine/Plugins/Developer/Concert/ConcertApp/MultiUserServer/Source/MultiUserServer/Private/Widgets/Clients/Browser/Graph/SClientNetworkGraphs.cpp
new file mode 100644
--- /dev/null
+++ b/Engine/Plugins/Developer/Concert/ConcertApp/MultiUserServer/Source/MultiUserServer/Private/Widgets/Clients/Browser/Graph/SClientNetworkGraphs.cpp
@@ -0,0 +1,134 @@
+// Copyright Epic Games, Inc. All Rights Reserved.
+
+#include "SClientNetworkGraphs.h"
+
+#include "ConcertServerStyle.h"
+#include "SNetworkGraph.h"
+#include "Widgets/Clients/Browser/Models/IClientTransferStatisticsModel.h"
+
+#include "Math/UnitConversion.h"
+#include "Widgets/Layout/SBorder.h"
+#include "Widgets/Layout/SSeparator.h"
+#include "Widgets/SBoxPanel.h"
+
+#define LOCTEXT_NAMESPACE "UnrealMultiUserUI.SClientNetworkGraphs"
+
+namespace UE::MultiUserServer 
+{
+	void SClientNetworkGraphs::Construct(const FArguments& InArgs, TSharedRef<IClientTransferStatisticsModel> InTransferStatsModel)
+	{
+		TransferStatisticsModel = MoveTemp(InTransferStatsModel);
+
+		ChildSlot
+		[
+			SNew(SBorder)
+			.BorderImage(FConcertServerStyle::Get().GetBrush("Concert.Clients.ThumbnailCurveBackground"))
+			[
+				BuildGraphWidgets(InArgs)
+			]
+		];
+	}
+
+	TSharedRef<SWidget> SClientNetworkGraphs::BuildGraphWidgets(const FArguments& InArgs)
+	{
+		static_assert(static_cast<int32>(EConcertTransferStatistic::Count) == 2, "Update the styles");
+		auto GetLineColor = [](EConcertTransferStatistic Stats)
+		{
+			switch (Stats)
+			{
+			case EConcertTransferStatistic::SentToClient: return FConcertServerStyle::Get().GetColor("Concert.Clients.NetworkGraph.Sent.LineColor");
+			case EConcertTransferStatistic::ReceivedFromClient: return FConcertServerStyle::Get().GetColor("Concert.Clients.NetworkGraph.Received.LineColor");
+			case EConcertTransferStatistic::Count:
+			default:
+				checkNoEntry();
+				return FLinearColor{};
+			}
+		};
+		auto GetFillColor = [](EConcertTransferStatistic Stats)
+		{
+			switch (Stats)
+			{
+			case EConcertTransferStatistic::SentToClient: return FConcertServerStyle::Get().GetColor("Concert.Clients.NetworkGraph.Sent.FillColor");
+			case EConcertTransferStatistic::ReceivedFromClient: return FConcertServerStyle::Get().GetColor("Concert.Clients.NetworkGraph.Received.FillColor");
+			case EConcertTransferStatistic::Count:
+			default:
+				checkNoEntry();
+				return FLinearColor{};
+			}
+		};
+		
+		TSharedRef<SVerticalBox> Root = SNew(SVerticalBox);
+
+		bool bNeedsSeparator = false;
+		for (int32 i = 0; i < static_cast<int32>(EConcertTransferStatistic::Count); ++i)
+		{
+			if (bNeedsSeparator)
+			{
+				Root->AddSlot()
+					.AutoHeight()
+					.Padding(0.f, 1.f, 0.f, 0.f)
+					[
+						SNew(SSeparator)
+						.Thickness(FConcertServerStyle::Get().GetFloat("Concert.Clients.NetworkGraph.GraphSeparatorLine.Thickness"))
+						.ColorAndOpacity(FConcertServerStyle::Get().GetColor("Concert.Clients.NetworkGraph.GraphSeparatorLine.LineColor"))
+					];
+			}
+			bNeedsSeparator = true;
+
+			const EConcertTransferStatistic Stat = static_cast<EConcertTransferStatistic>(i);
+			TSharedPtr<SNetworkGraph> Graph;
+			Root->AddSlot()
+			[
+				SAssignNew(Graph, SNetworkGraph)
+				.RequestViewUpdate_Lambda([this, Stat]()
+				{
+					OnTransferStatisticsUpdated(Stat);
+				})
+				.CurrentNetworkValue(this, &SClientNetworkGraphs::GetStatTextValue, Stat)
+				.DisplayedTimeRange(InArgs._TimeRange)
+				.CurveColor(GetLineColor(Stat))
+				.FillColor(GetFillColor(Stat))
+			];
+			Graphs.Add(Stat, Graph);
+			TransferStatisticsModel->OnTransferTimelineUpdated(Stat).AddSP(this, &SClientNetworkGraphs::OnTransferStatisticsUpdated, Stat);
+		}
+
+		return Root;
+	}
+
+	FText SClientNetworkGraphs::GetStatTextValue(EConcertTransferStatistic Stat) const
+	{
+		const FString Label = [Stat]()
+		{
+			switch (Stat)
+			{
+			case EConcertTransferStatistic::SentToClient: return TEXT("S");
+			case EConcertTransferStatistic::ReceivedFromClient: return TEXT("R");
+			case EConcertTransferStatistic::Count: 
+			default:
+				checkNoEntry();
+				return TEXT("Null");
+			}
+		}();
+
+		const TArray<FConcertTransferSamplePoint>& TransferStatistics = TransferStatisticsModel->GetTransferStatTimeline(Stat);
+		if (TransferStatistics.Num() > 0)
+		{
+			const uint64 BytesSent = TransferStatistics[TransferStatistics.Num() - 1].BytesTransferred;
+			const FNumericUnit<uint64> TargetUnit = FUnitConversion::QuantizeUnitsToBestFit(BytesSent, EUnit::Bytes);
+			return FText::Format(LOCTEXT("DataTextFmt", "{0}: {1} {2}"), FText::FromString(Label), TargetUnit.Value, FText::FromString(FUnitConversion::GetUnitDisplayString(TargetUnit.Units)));
+		}
+		
+		return FText::Format(LOCTEXT("NoDataTextFmt", "{0}: n/a"), FText::FromString(Label));
+	}
+
+	void SClientNetworkGraphs::OnTransferStatisticsUpdated(EConcertTransferStatistic Stat)
+	{
+		const TArray<FConcertTransferSamplePoint>& Timeline = TransferStatisticsModel->GetTransferStatTimeline(Stat);
+		Graphs[Stat]->UpdateView(Timeline.Num(), [&Timeline](uint32 Index)
+		{
+			return SNetworkGraph::FGraphData{ Timeline[Index].BytesTransferred, Timeline[Index].LocalTime };
+		});
+	}
+}
+#undef LOCTEXT_NAMESPACE
diff --git a/Engine/Plugins/Developer/Concert/ConcertApp/MultiUserServer/Source/MultiUserServer/Private/Widgets/Clients/Browser/Graph/SClientNetworkGraphs.h b/Engine/Plugins/Developer/Concert/ConcertApp/MultiUserServer/Source/MultiUserServer/Private/Widgets/Clients/Browser/Graph/SClientNetworkGraphs.h
new file mode 100644
--- /dev/null
+++ b/Engine/Plugins/Developer/Concert/ConcertApp/MultiUserServer/Source/MultiUserServer/Private/Widgets/Clients/Browser/Graph/SClientNetworkGraphs.h
@@ -0,0 +1,41 @@
+// Copyright Epic Games, Inc. All Rights Reserved.
+
+#pragma once
+
+#include "CoreMinimal.h"
+#include "Widgets/DeclarativeSyntaxSupport.h"
+#include "Widgets/SCompoundWidget.h"
+#include "Widgets/Clients/Browser/Models/IClientTransferStatisticsModel.h"
+
+namespace UE::MultiUserServer
+{
+	class IClientTransferStatisticsModel;
+	class SNetworkGraph;
+
+	/** Displays send and receive graphs */
+	class SClientNetworkGraphs : public SCompoundWidget
+	{
+	public:
+
+		SLATE_BEGIN_ARGS(SClientNetworkGraphs)
+			: _TimeRange(FTimespan::FromSeconds(15))
+		{}
+			/** The last x time of traffic to display*/
+			SLATE_ARGUMENT(FTimespan, TimeRange)
+		SLATE_END_ARGS()
+
+		void Construct(const FArguments& InArgs, TSharedRef<IClientTransferStatisticsModel> InTransferStatsModel);
+
+	private:
+
+		/** Tells us the data to display */
+		TSharedPtr<IClientTransferStatisticsModel> TransferStatisticsModel;
+
+		TMap<EConcertTransferStatistic, TSharedPtr<SNetworkGraph>> Graphs;
+		
+		TSharedRef<SWidget> BuildGraphWidgets(const FArguments& InArgs);
+		FText GetStatTextValue(EConcertTransferStatistic Stat) const;
+		
+		void OnTransferStatisticsUpdated(EConcertTransferStatistic Stat);
+	};
+}
diff --git a/Engine/Plugins/Developer/Concert/ConcertApp/MultiUserServer/Source/MultiUserServer/Private/Widgets/Clients/Browser/Graph/SCurveTimelineView.cpp b/Engine/Plugins/Developer/Concert/ConcertApp/MultiUserServer/Source/MultiUserServer/Private/Widgets/Clients/Browser/Graph/SCurveTimelineView.cpp
new file mode 100644
--- /dev/null
+++ b/Engine/Plugins/Developer/Concert/ConcertApp/MultiUserServer/Source/MultiUserServer/Private/Widgets/Clients/Browser/Graph/SCurveTimelineView.cpp
@@ -0,0 +1,189 @@
+// Copyright Epic Games, Inc. All Rights Reserved.
+
+#include "SCurveTimelineView.h"
+#include "Framework/Application/SlateApplication.h"
+#include "Rendering/DrawElements.h"
+#include "SSimpleTimeSlider.h"
+
+#define LOCTEXT_NAMESPACE "SCurveTimelineView"
+
+
+void SCurveTimelineView::Construct( const SCurveTimelineView::FArguments& InArgs )
+{
+    ViewRange = InArgs._ViewRange;
+	CurveColor = InArgs._CurveColor;
+	FillColor = InArgs._FillColor;
+	RenderFill = InArgs._RenderFill;
+	DesiredSize = InArgs._DesiredSize;
+	CurveData = InArgs._CurveData;
+
+	// set clipping on by default, since the OnPaint function is drawing outside the bounds
+	Clipping = EWidgetClipping::ClipToBounds;
+}
+
+int32 SCurveTimelineView::OnPaint( const FPaintArgs& Args, const FGeometry& AllottedGeometry, const FSlateRect& MyCullingRect, FSlateWindowElementList& OutDrawElements, int32 LayerId, const FWidgetStyle& InWidgetStyle, bool bParentEnabled ) const
+{
+	int32 NewLayer = PaintCurve( AllottedGeometry, MyCullingRect, OutDrawElements, LayerId, InWidgetStyle, bParentEnabled );
+
+	return FMath::Max( NewLayer, SCompoundWidget::OnPaint(Args, AllottedGeometry, MyCullingRect, OutDrawElements, NewLayer, InWidgetStyle, ShouldBeEnabled( bParentEnabled ) ) );
+}
+
+
+FVector2D SCurveTimelineView::ComputeDesiredSize( float ) const
+{
+	return DesiredSize.Get();
+}
+
+int32 SCurveTimelineView::PaintCurve(const FGeometry& AllottedGeometry, const FSlateRect& MyCullingRect, FSlateWindowElementList& OutDrawElements, int32 LayerId, const FWidgetStyle& InWidgetStyle, bool bParentEnabled ) const
+{
+	static const float LargeFrameTime = 1.125;  // ideally, we could check the recorded frame data for actual frame length
+	 
+	// convert time range to from rewind debugger times to profiler times
+	TRange<double> DebugTimeRange = ViewRange.Get();
+	FLinearColor LineColor = CurveColor.Get();
+	
+	SSimpleTimeSlider::FScrubRangeToScreen RangeToScreen( DebugTimeRange, AllottedGeometry.GetLocalSize() );
+	FVector2D Size = AllottedGeometry.GetLocalSize();
+
+	if (TSharedPtr<FTimelineCurveData> Curve = CurveData.Get())
+	{
+		// Min/Max for normalization
+		// maybe expose these as attributes too, so they are recomputed less frequently
+		float MinValue = 1.0e10;
+		float MaxValue = -1.0e10;
+
+		if (bUseFixedRange)
+		{
+			MinValue = FixedRangeMin;
+			MaxValue = FixedRangeMax;
+		}
+		else
+		{
+			for(auto& Point : Curve->Points)
+			{
+				MinValue = FMath::Min(MinValue, Point.Value);
+				MaxValue = FMath::Max(MaxValue, Point.Value);
+			}
+		}
+		
+		float Range = FMath::Max(1.0,MaxValue - MinValue);
+		
+		auto &CurvePoints = Curve->Points;
+		const int NumPoints = CurvePoints.Num();
+		
+		const ESlateDrawEffect LineDrawEffects = ESlateDrawEffect::NoBlending | ESlateDrawEffect::NoPixelSnapping;	
+
+		if (NumPoints > 0)
+		{
+			if (RenderFill.Get())
+			{
+				// don't call this every frame
+				FSlateResourceHandle ResourceHandle = FSlateApplication::Get().GetRenderer()->GetResourceHandle(*FAppStyle::GetBrush("Sequencer.SectionArea.Background"));
+			
+				const FSlateRenderTransform& Transform = AllottedGeometry.GetAccumulatedRenderTransform();
+				FColor Color = FillColor.Get().ToFColor(true);
+
+				// filled polygons
+				TArray<FSlateVertex> Vertices;
+				TArray<SlateIndex> Indices;
+				Vertices.Reserve((NumPoints-1)*4);
+				Indices.Reserve((NumPoints-1)*6);
+				
+				for(int i=1; i<NumPoints; i++)
+				{
+					const FTimelineCurveData::FCurvePoint& Point1 = CurvePoints[i-1];
+					const FTimelineCurveData::FCurvePoint& Point2 = CurvePoints[i];
+				
+					if (Point2.Time - Point1.Time < LargeFrameTime)
+					{
+						float X1 = RangeToScreen.InputToLocalX(Point1.Time);
+						const float NormalizedValue1 = (Point1.Value - MinValue) / Range;
+						float Y1 = Size.Y - (Size.Y - 2) * NormalizedValue1;
+						float X2 = RangeToScreen.InputToLocalX(Point2.Time);
+						const float NormalizedValue2 = (Point2.Value - MinValue) / Range;
+						float Y2 = Size.Y - (Size.Y - 2) * NormalizedValue2;
+					
+						Vertices.Add(FSlateVertex::Make<ESlateVertexRounding::Disabled>(Transform, FVector2f(X1,Size.Y-1) ,FVector2f(0,0), Color));
+						Vertices.Add(FSlateVertex::Make<ESlateVertexRounding::Disabled>(Transform, FVector2f(X1,Y1),FVector2f(0,0), Color));
+						Vertices.Add(FSlateVertex::Make<ESlateVertexRounding::Disabled>(Transform, FVector2f(X2,Y2),FVector2f(0,0), Color));
+						Vertices.Add(FSlateVertex::Make<ESlateVertexRounding::Disabled>(Transform, FVector2f(X2,Size.Y-1),FVector2f(0,0), Color));
+						
+						int32 Index0 = Vertices.Num() - 4;
+                        int32 Index1 = Vertices.Num() - 3;
+                        int32 Index2 = Vertices.Num() - 2;
+                        int32 Index3 = Vertices.Num() - 1;
+                       
+						Indices.Add(Index0);
+						Indices.Add(Index1);
+						Indices.Add(Index2);
+						Indices.Add(Index0);
+						Indices.Add(Index2);
+						Indices.Add(Index3);
+					}
+					
+					FSlateDrawElement::MakeCustomVerts(
+								OutDrawElements,
+										LayerId++,
+										ResourceHandle,
+										Vertices,
+										Indices,
+										nullptr,
+										0,
+										0);
+				}
+			}
+
+
+			// Lines rendering
+			
+			TArray<FVector2D> Points;
+			Points.Reserve(NumPoints);
+
+			float PrevTime = CurvePoints[0].Time;
+			for(int i=0; i<NumPoints; i++)
+			{
+				const FTimelineCurveData::FCurvePoint& Point = CurvePoints[i];
+				if (Point.Time - PrevTime > LargeFrameTime && Points.Num()>1)
+				{
+					// break the line list - data has stopped and started again
+					FSlateDrawElement::MakeLines(
+								OutDrawElements,
+										LayerId++,
+										AllottedGeometry.ToPaintGeometry(),
+										Points,
+										LineDrawEffects,
+										LineColor,
+										false
+										);
+					Points.SetNum(0,false);
+				}
+			
+				float X = RangeToScreen.InputToLocalX(Point.Time);
+				PrevTime = Point.Time;
+				const float NormalizedValue = (Point.Value - MinValue) / Range;
+				float Y = Size.Y - (Size.Y - 2) * NormalizedValue;
+				Points.Add(FVector2D(X,Y));
+			}
+			
+			FSlateDrawElement::MakeLines(
+						OutDrawElements,
+								LayerId,
+								AllottedGeometry.ToPaintGeometry(),
+								Points,
+								LineDrawEffects,
+								LineColor,
+								true
+								);
+
+
+								
+		}
+	
+
+
+	}
+
+	return LayerId;
+}
+
+#undef LOCTEXT_NAMESPACE
diff --git a/Engine/Plugins/Developer/Concert/ConcertApp/MultiUserServer/Source/MultiUserServer/Private/Widgets/Clients/Browser/Graph/SCurveTimelineView.h b/Engine/Plugins/Developer/Concert/ConcertApp/MultiUserServer/Source/MultiUserServer/Private/Widgets/Clients/Browser/Graph/SCurveTimelineView.h
new file mode 100644
--- /dev/null
+++ b/Engine/Plugins/Developer/Concert/ConcertApp/MultiUserServer/Source/MultiUserServer/Private/Widgets/Clients/Browser/Graph/SCurveTimelineView.h
@@ -0,0 +1,94 @@
+// Copyright Epic Games, Inc. All Rights Reserved.
+
+#pragma once
+
+#include "CoreMinimal.h"
+#include "Widgets/SCompoundWidget.h"
+#include "Widgets/DeclarativeSyntaxSupport.h"
+
+class FPaintArgs;
+class FSlateWindowElementList;
+class FTimeSliderController;
+
+/**
+ * IMPORTANT: This code is duplicated and slighlty modified from Plugins/Animation/GameplayInsights.
+ * The original widget code is supposed to make its way into ToolWidgets at which point this duplicate code should be removed.
+ */
+
+
+class SCurveTimelineView : public SCompoundWidget
+{
+public:
+	struct FTimelineCurveData
+	{
+		struct FCurvePoint
+		{
+			double Time;
+			float Value;
+		};
+
+		TArray<FCurvePoint> Points;
+	};
+	
+	SLATE_BEGIN_ARGS(SCurveTimelineView)
+		: _ViewRange(TRange<double>(0,10))
+		, _DesiredSize(FVector2D(100.f,20.f))
+		, _RenderFill(false)
+	{}
+    	/** View time range */
+    	SLATE_ATTRIBUTE(TRange<double>, ViewRange);
+
+		/** Data for curve to render */
+		SLATE_ATTRIBUTE(TSharedPtr<FTimelineCurveData>, CurveData);
+
+		/** Desired widget size */
+		SLATE_ATTRIBUTE(FVector2D, DesiredSize);
+	
+		/** Curve Color */
+		SLATE_ATTRIBUTE(FLinearColor, CurveColor);
+	
+		/** Fill Color */
+		SLATE_ATTRIBUTE(FLinearColor, FillColor);
+	
+		/** Enable Fill rendering */
+		SLATE_ATTRIBUTE(bool, RenderFill);
+     
+	SLATE_END_ARGS()
+
+
+	/**
+	 * Construct the widget
+	 * 
+	 * @param InArgs   A declaration from which to construct the widget
+	 */
+	void Construct( const FArguments& InArgs );
+
+	// render the curve with a fixed min/max range, rather than using the min/max from data
+	void SetFixedRange(float Min, float Max)
+	{
+		bUseFixedRange = true;
+		FixedRangeMin = Min;
+		FixedRangeMax = Max;
+	}
+
+protected:
+	// SWidget interface
+	virtual FVector2D ComputeDesiredSize(float) const override;
+	virtual int32 OnPaint( const FPaintArgs& Args, const FGeometry& AllottedGeometry, const FSlateRect& MyCullingRect, FSlateWindowElementList& OutDrawElements, int32 LayerId, const FWidgetStyle& InWidgetStyle, bool bParentEnabled ) const override;
+	
+	int32 PaintCurve(const FGeometry& AllottedGeometry, const FSlateRect& MyCullingRect, FSlateWindowElementList& OutDrawElements, int32 LayerId, const FWidgetStyle& InWidgetStyle, bool bParentEnabled ) const;
+	
+	TAttribute<TRange<double>> ViewRange;
+	TAttribute<TSharedPtr<FTimelineCurveData>> CurveData;
+	TAttribute<FVector2D> DesiredSize;
+	TAttribute<FLinearColor> CurveColor;
+	TAttribute<FLinearColor> FillColor;
+	TAttribute<bool> RenderFill;
+	
+	const FSlateBrush* WhiteBrush;
+
+	bool bUseFixedRange;
+	float FixedRangeMin;
+	float FixedRangeMax;
+};
+
diff --git a/Engine/Plugins/Developer/Concert/ConcertApp/MultiUserServer/Source/MultiUserServer/Private/Widgets/Clients/Browser/Graph/SNetworkGraph.cpp b/Engine/Plugins/Developer/Concert/ConcertApp/MultiUserServer/Source/MultiUserServer/Private/Widgets/Clients/Browser/Graph/SNetworkGraph.cpp
new file mode 100644
--- /dev/null
+++ b/Engine/Plugins/Developer/Concert/ConcertApp/MultiUserServer/Source/MultiUserServer/Private/Widgets/Clients/Browser/Graph/SNetworkGraph.cpp
@@ -0,0 +1,247 @@
+// Copyright Epic Games, Inc. All Rights Reserved.
+
+#include "SNetworkGraph.h"
+
+#include "ConcertServerStyle.h"
+#include "Math/UnitConversion.h"
+
+#include "Widgets/Layout/SSeparator.h"
+#include "Widgets/SBoxPanel.h"
+#include "Widgets/SNullWidget.h"
+#include "Widgets/SOverlay.h"
+#include "Widgets/Text/STextBlock.h"
+
+#define LOCTEXT_NAMESPACE "UnrealMultiUserUI.SNetworkGraph"
+
+namespace UE::MultiUserServer
+{
+	void SNetworkGraph::Construct(const FArguments& InArgs)
+	{
+		check(InArgs._DisplayedTimeRange.IsBound() || InArgs._DisplayedTimeRange.IsSet());
+		check(InArgs._RequestViewUpdate.IsBound());
+		check(InArgs._TimeSliceSize > 0.f);
+		
+		DisplayedTimeRange = InArgs._DisplayedTimeRange;
+		RequestViewUpdate = InArgs._RequestViewUpdate;
+		MinimumYAxisRange = InArgs._MinimumYAxisRange;
+		TimeSliceSize = InArgs._TimeSliceSize;
+		TimeBucketAnchor = FDateTime::Now();
+		
+		CurveData = MakeShared<SCurveTimelineView::FTimelineCurveData>();
+		
+		ChildSlot
+		[
+			SNew(SOverlay)
+	
+			+SOverlay::Slot()
+			.Padding(0.f, 5.f, 0.f, 0.f)
+			[
+				SAssignNew(Graph, SCurveTimelineView)
+				.ViewRange(TRange<double>{ 0.0, DisplayedTimeRange.Get().GetTotalSeconds() })
+				.CurveData(CurveData)
+				.CurveColor(InArgs._CurveColor)
+				.FillColor(InArgs._FillColor)
+				.RenderFill(true)
+			]
+	
+			+SOverlay::Slot()
+			[
+				SNew(SHorizontalBox)
+				+SHorizontalBox::Slot()
+				.VAlign(VAlign_Top)
+				.HAlign(HAlign_Right)
+				[
+					SNew(STextBlock)
+					.Text(InArgs._CurrentNetworkValue)
+					.ColorAndOpacity(FColor::White)
+				]
+			]
+			
+			+SOverlay::Slot()
+			[
+				CreateHorizontalLines()
+			]
+		];
+
+		SetTimerForForceRefresh();
+	}
+
+	using FCurvePoint = SCurveTimelineView::FTimelineCurveData::FCurvePoint;
+	/** Distributes the points into time slices so the graph consists of n segements across the x-axis*/
+	static TArray<FCurvePoint> DistributeIntoTimeBuckets(TArrayView<FCurvePoint> Curve, float TimeRange, float TimeSliceSize, const FDateTime& Now, const FDateTime& TimeAnchor);
+	/** Makes a curve with corners flat so there are no longer an spiky corners on it. Imagine it like a curve of rectangles. */
+	static TArray<FCurvePoint> FlattenCurve(TArrayView<FCurvePoint> Curve);
+
+	void SNetworkGraph::UpdateView(uint32 Num, FGetGraphData GetDataFunc)
+	{
+		const FTimespan TimeRange = DisplayedTimeRange.Get();
+		const double DisplayTimeRangInSeconds = TimeRange.GetTotalSeconds();
+		const FDateTime LocalNow = FDateTime::Now();
+		LastViewUpdate = LocalNow;
+
+		TArray<FCurvePoint> Points;
+		MaxDataInDisplayTime = MinimumYAxisRange;
+		for (uint32 i = 0; i < Num; ++i)
+		{
+			const FGraphData GraphData = GetDataFunc(i);
+			const FTimespan TimeSinceSample = LocalNow - GraphData.LocalTime;
+			if (TimeSinceSample <= TimeRange)
+			{
+				MaxDataInDisplayTime = FMath::Max(MaxDataInDisplayTime, GraphData.Data);
+				const double TimeLocation = DisplayTimeRangInSeconds - TimeSinceSample.GetTotalSeconds();
+				const float Value = GraphData.Data;
+			
+				const SCurveTimelineView::FTimelineCurveData::FCurvePoint Point{ TimeLocation, Value };
+				Points.Add(Point);
+			}
+		}
+		
+		Points = DistributeIntoTimeBuckets(Points, DisplayTimeRangInSeconds, TimeSliceSize, LocalNow, TimeBucketAnchor);
+		CurveData->Points = FlattenCurve(Points);
+		Graph->SetFixedRange(0, MaxDataInDisplayTime);
+	}
+
+	static TArray<FCurvePoint> DistributeIntoTimeBuckets(TArrayView<FCurvePoint> Curve, float TimeRange, float TimeSliceSize, const FDateTime& Now, const FDateTime& TimeAnchor)
+	{
+		// The x-axis is divided into n equal segments of size TimeSliceSize
+
+		// Step 1: Preallocate result
+		const int32 ResultNum = TimeRange / TimeSliceSize;
+		TArray<FCurvePoint> Result;
+		Result.SetNumUninitialized(ResultNum);
+		for (int32 ResultIdx = 0; ResultIdx < ResultNum; ++ResultIdx)
+		{
+			Result[ResultIdx] = { static_cast<float>(ResultIdx) * TimeSliceSize, 0 };
+		}
+		
+		// Step 2: The times must be offset from a global time anchor point.
+		// If we simply used the measurement time, we'd implicity use Now as time anchor;
+		// That would mean that certain points would change buckets to the left or right as the graph is plotted over time (comment this code to see the effect).
+		const FTimespan SinceAnchor = Now - TimeAnchor;
+		const double BucketsSinceAnchor = SinceAnchor.GetTotalSeconds() / TimeSliceSize;
+		const int32 NextBucketIndexAfterNow = static_cast<int32>(BucketsSinceAnchor) + 1;
+		const FTimespan TimeUntilNextBucketAfterNow = FTimespan::FromSeconds(NextBucketIndexAfterNow * TimeSliceSize);
+		const FDateTime StartTimeOfNextBucketAfterNow = TimeAnchor + TimeUntilNextBucketAfterNow;
+		// Negative not positive: the offset time should end up in previous buckets... not in the next bucket which lies in the future
+		const double TimeToOffset = -(StartTimeOfNextBucketAfterNow - Now).GetTotalSeconds();
+
+		// Step 3: Sort the times into time buckets offsetting the times relative to TimeAnchor
+		for (const FCurvePoint& CurvePoint : Curve)
+		{
+			// Due to micro gaps, e.g. 7.99999 and 9.00001 and sampling every second would assign
+			// a value of 0 to the index corresponding to 8-9 seconds.
+			// This is fixed by rounding the time.
+			const int32 Index = FMath::Clamp<int32>(FMath::RoundToFloat(CurvePoint.Time + TimeToOffset) / TimeSliceSize, 0, ResultNum - 1);
+			Result[Index].Value += CurvePoint.Value;
+		}
+		
+		return Result;
+	}
+	
+	static TArray<FCurvePoint> FlattenCurve(TArrayView<FCurvePoint> Curve)
+	{
+		TArray<FCurvePoint> Result;
+		for (int32 i = 0; i < Curve.Num() - 1; ++i)
+		{
+			const FCurvePoint& Left = Curve[i];
+			const FCurvePoint& Right = Curve[i + 1];
+
+			Result.Add({ Left.Time, Left.Value });
+			Result.Add({ Left.Time, Right.Value });
+			Result.Add({ Right.Time, Right.Value });
+		}
+		return Result;
+	}
+
+	TSharedRef<SWidget> SNetworkGraph::CreateHorizontalLines() const
+	{
+		// There should be a line at 1/3 and 2/3 of the graph height (to guide the user's eye)
+		// The line slots must take 1/3 of the space each while the top and the bottom buffer slots must take 1/6 each
+		constexpr float LineFillHeight = 2; // = 0.3333f * 6
+		constexpr float BufferFillHeight = 1; // = 0.1666f * 6 
+		return SNew(SVerticalBox)
+		
+			+SVerticalBox::Slot()
+			.FillHeight(BufferFillHeight)
+			[
+				SNullWidget::NullWidget
+			]
+		
+			+SVerticalBox::Slot()
+			.FillHeight(LineFillHeight)
+			[
+				SNew(SHorizontalBox)
+
+				+SHorizontalBox::Slot()
+				.AutoWidth()
+				.VAlign(VAlign_Center)
+				[
+					SNew(STextBlock)
+					.Text_Lambda([this](){ return GetBytesDisplayText(2 * MaxDataInDisplayTime / 3); })
+					.ColorAndOpacity(FConcertServerStyle::Get().GetColor("Concert.Clients.NetworkGraph.HorizontalHelperLine.TextColor"))
+				]
+
+				+SHorizontalBox::Slot()
+				.FillWidth(1.f)
+				.Padding(2.f, 0.f, 0.f, 0.f)
+				.VAlign(VAlign_Center)
+				[
+					SNew(SSeparator)
+					.Thickness(FConcertServerStyle::Get().GetFloat("Concert.Clients.NetworkGraph.HorizontalHelperLine.Thickness"))
+					.ColorAndOpacity(FConcertServerStyle::Get().GetColor("Concert.Clients.NetworkGraph.HorizontalHelperLine.LineColor"))
+				]
+			]
+
+			+SVerticalBox::Slot()
+			.FillHeight(LineFillHeight)
+			[
+				SNew(SHorizontalBox)
+
+				+SHorizontalBox::Slot()
+				.AutoWidth()
+				.VAlign(VAlign_Center)
+				[
+					SNew(STextBlock)
+					.Text_Lambda([this](){ return GetBytesDisplayText(MaxDataInDisplayTime / 3); })
+					.ColorAndOpacity(FConcertServerStyle::Get().GetColor("Concert.Clients.NetworkGraph.HorizontalHelperLine.TextColor"))
+				]
+
+				+SHorizontalBox::Slot()
+				.FillWidth(1.f)
+				.Padding(2.f, 0.f, 0.f, 0.f)
+				.VAlign(VAlign_Center)
+				[
+					SNew(SSeparator)
+					.Thickness(FConcertServerStyle::Get().GetFloat("Concert.Clients.NetworkGraph.HorizontalHelperLine.Thickness"))
+					.ColorAndOpacity(FConcertServerStyle::Get().GetColor("Concert.Clients.NetworkGraph.HorizontalHelperLine.LineColor"))
+				]
+			]
+
+			+SVerticalBox::Slot()
+			.FillHeight(BufferFillHeight)
+			[
+				SNullWidget::NullWidget
+			];
+	}
+
+	FText SNetworkGraph::GetBytesDisplayText(uint64 Data) const
+	{
+		const FNumericUnit<uint64> TargetUnit = FUnitConversion::QuantizeUnitsToBestFit(Data, EUnit::Bytes);
+		return FText::Format(LOCTEXT("BytesFmt", "{0} {1}"), TargetUnit.Value, FText::FromString(FUnitConversion::GetUnitDisplayString(TargetUnit.Units)));
+	}
+
+	void SNetworkGraph::SetTimerForForceRefresh()
+	{
+		RegisterActiveTimer(TimeSliceSize, FWidgetActiveTimerDelegate::CreateLambda([this](double InCurrentTime, float InDeltaTime)
+		{
+			const FTimespan TimeSinceLastUpdate = LastViewUpdate - FDateTime::Now();
+			if (TimeSinceLastUpdate.GetTotalSeconds() < TimeSliceSize)
+			{
+				RequestViewUpdate.Execute();
+			}
+			return EActiveTimerReturnType::Continue;
+		}));
+	}
+}
+
+#undef LOCTEXT_NAMESPACE
\ No newline at end of file
diff --git a/Engine/Plugins/Developer/Concert/ConcertApp/MultiUserServer/Source/MultiUserServer/Private/Widgets/Clients/Browser/Graph/SNetworkGraph.h b/Engine/Plugins/Developer/Concert/ConcertApp/MultiUserServer/Source/MultiUserServer/Private/Widgets/Clients/Browser/Graph/SNetworkGraph.h
new file mode 100644
--- /dev/null
+++ b/Engine/Plugins/Developer/Concert/ConcertApp/MultiUserServer/Source/MultiUserServer/Private/Widgets/Clients/Browser/Graph/SNetworkGraph.h
@@ -0,0 +1,92 @@
+// Copyright Epic Games, Inc. All Rights Reserved.
+
+#pragma once
+
+#include "CoreMinimal.h"
+#include "SCurveTimelineView.h"
+#include "Widgets/DeclarativeSyntaxSupport.h"
+#include "Widgets/SCompoundWidget.h"
+
+namespace UE::MultiUserServer
+{
+	/**
+	 * Displays a graph from a source of floats
+	 */
+	class SNetworkGraph : public SCompoundWidget
+	{
+	public:
+
+		DECLARE_DELEGATE(FRequestViewUpdate);
+		
+		struct FGraphData
+		{
+			uint64 Data;
+			FDateTime LocalTime;
+		};
+		
+		using FGetGraphData = TFunctionRef<FGraphData(uint32 Index)>;
+
+		SLATE_BEGIN_ARGS(SNetworkGraph)
+			: _MinimumYAxisRange(FMath::Pow(2.f, 13.f)) // 8 kilobytes
+			, _TimeSliceSize(1.f)
+		{}
+			/** Displayed at the top-right of the graph, e.g. S: 128 kB */
+			SLATE_ATTRIBUTE(FText, CurrentNetworkValue)
+			/** All data within the last DisplayedTimeRange.GetTotalSeconds() will be displayed. */
+			SLATE_ATTRIBUTE(FTimespan, DisplayedTimeRange)
+
+			SLATE_ARGUMENT(FLinearColor, CurveColor)
+			SLATE_ARGUMENT(FLinearColor, FillColor)
+			SLATE_ARGUMENT(uint64, MinimumYAxisRange)
+			/** The curve is divided into even segments along the x-axis. This specifies the number of seconds each segment is wide. */
+			SLATE_ARGUMENT(float, TimeSliceSize)
+
+			/** Called when it has been too long since UpdateView has last been called. Causes UpdateView to be called. */
+			SLATE_EVENT(FRequestViewUpdate, RequestViewUpdate)
+		SLATE_END_ARGS()
+
+		void Construct(const FArguments& InArgs);
+
+		/** Updates the graph view to display the given data. */
+		void UpdateView(uint32 Num, FGetGraphData GetDataFunc);
+
+	private:
+
+		/** All data within the last DisplayedTimeRange.GetTotalSeconds() will be displayed. */
+		TAttribute<FTimespan> DisplayedTimeRange;
+		/** Called when it has been too long since UpdateView has last been called. Causes UpdateView to be called. */
+		FRequestViewUpdate RequestViewUpdate;
+
+		/** Renders the graph */
+		TSharedPtr<SCurveTimelineView> Graph;
+		/** The data displayed in Graph*/
+		TSharedPtr<SCurveTimelineView::FTimelineCurveData> CurveData;
+		
+		/** Heightest point in CurveData. Used to correctly set y-axis scale */
+		uint64 MaxDataInDisplayTime;
+		/** Used to regularily request updates if the network does not provide them (e.g. because there just was no traffic). */
+		FDateTime LastViewUpdate;
+		
+		/**
+		 * The graph is constantly scaled to fit the biggest value in the point set.
+		 * This variable artifically sets the minimum value to which we scale.
+		 */
+		uint64 MinimumYAxisRange;
+		/**
+		 * The curve is divided into even segments along the x-axis.
+		 * This specifies the number of seconds each segment is wide.
+		 */
+		float TimeSliceSize;
+		/** All time slice buckets are relative to this time which is taken when the widget is constructed. */
+		FDateTime TimeBucketAnchor;
+		
+
+		TSharedRef<SWidget> CreateHorizontalLines() const;
+		FText GetBytesDisplayText(uint64 Data) const;
+
+		/** Updates the graph in an interval of TimeSliceSize in case no network data is received in that time. */
+		void SetTimerForForceRefresh();
+	};
+}
+
+
diff --git a/Engine/Plugins/Developer/Concert/ConcertApp/MultiUserServer/Source/MultiUserServer/Private/Widgets/Clients/Browser/Models/ClientTransferStatisticsModel.cpp b/Engine/Plugins/Developer/Concert/ConcertApp/MultiUserServer/Source/MultiUserServer/Private/Widgets/Clients/Browser/Models/ClientTransferStatisticsModel.cpp
--- a/Engine/Plugins/Developer/Concert/ConcertApp/MultiUserServer/Source/MultiUserServer/Private/Widgets/Clients/Browser/Models/ClientTransferStatisticsModel.cpp
+++ b/Engine/Plugins/Developer/Concert/ConcertApp/MultiUserServer/Source/MultiUserServer/Private/Widgets/Clients/Browser/Models/ClientTransferStatisticsModel.cpp
@@ -6,88 +6,217 @@
 #include "Containers/Ticker.h"
 #include "Features/IModularFeatures.h"
 
-namespace UE::MultiUserServer::Private::ClientTransferStatisticsModel
+namespace UE::MultiUserServer
 {
-	static INetworkMessagingExtension* GetMessagingStatistics()
+	namespace Private::ClientTransferStatisticsModel
 	{
-		IModularFeatures& ModularFeatures = IModularFeatures::Get();
-		if (IsInGameThread())
+		static INetworkMessagingExtension* GetMessagingStatistics()
 		{
-			if (ModularFeatures.IsModularFeatureAvailable(INetworkMessagingExtension::ModularFeatureName))
+			IModularFeatures& ModularFeatures = IModularFeatures::Get();
+			if (IsInGameThread())
 			{
-				return &ModularFeatures.GetModularFeature<INetworkMessagingExtension>(INetworkMessagingExtension::ModularFeatureName);
+				if (ModularFeatures.IsModularFeatureAvailable(INetworkMessagingExtension::ModularFeatureName))
+				{
+					return &ModularFeatures.GetModularFeature<INetworkMessagingExtension>(INetworkMessagingExtension::ModularFeatureName);
+				}
 			}
-		}
-		else
-		{
-			ModularFeatures.LockModularFeatureList();
-			ON_SCOPE_EXIT
+			else
 			{
-				ModularFeatures.UnlockModularFeatureList();
-			};
+				ModularFeatures.LockModularFeatureList();
+				ON_SCOPE_EXIT
+				{
+					ModularFeatures.UnlockModularFeatureList();
+				};
 			
-			if (ModularFeatures.IsModularFeatureAvailable(INetworkMessagingExtension::ModularFeatureName))
-			{
-				return &ModularFeatures.GetModularFeature<INetworkMessagingExtension>(INetworkMessagingExtension::ModularFeatureName);
+				if (ModularFeatures.IsModularFeatureAvailable(INetworkMessagingExtension::ModularFeatureName))
+				{
+					return &ModularFeatures.GetModularFeature<INetworkMessagingExtension>(INetworkMessagingExtension::ModularFeatureName);
+				}
 			}
+		
+			ensureMsgf(false, TEXT("Feature %s is unavailable"), *INetworkMessagingExtension::ModularFeatureName.ToString());
+			return nullptr;
+		}
+	}
+
+	FClientTransferStatisticsModel::FClientTransferStatisticsModel(const FMessageAddress& ClientAddress)
+		: ClientAddress(ClientAddress)
+	{
+		for (int32 i = 0; i < static_cast<int32>(EConcertTransferStatistic::Count); ++i)
+		{
+			const EConcertTransferStatistic StatType = static_cast<EConcertTransferStatistic>(i);
+			TransferStatisticsTimelines.Add(StatType);
+			OnTimelineUpdatedDelegates.Add(StatType);
 		}
 		
-		ensureMsgf(false, TEXT("Feature %s is unavailable"), *INetworkMessagingExtension::ModularFeatureName.ToString());
-		return nullptr;
+		if (INetworkMessagingExtension* Statistics = Private::ClientTransferStatisticsModel::GetMessagingStatistics())
+		{
+			Statistics->OnTransferUpdatedFromThread().AddRaw(this, &FClientTransferStatisticsModel::OnTransferUpdatedFromThread);
+			TickHandle = FTSTicker::GetCoreTicker().AddTicker(FTickerDelegate::CreateRaw(this, &FClientTransferStatisticsModel::Tick));
+		}
 	}
-}
 
-UE::MultiUserServer::FClientTransferStatisticsModel::FClientTransferStatisticsModel(const FMessageAddress& MessageAddress)
-	: MessageAddress(MessageAddress)
-{
-	if (INetworkMessagingExtension* Statistics = Private::ClientTransferStatisticsModel::GetMessagingStatistics())
+	FClientTransferStatisticsModel::~FClientTransferStatisticsModel()
 	{
-		Statistics->OnTransferUpdatedFromThread().AddRaw(this, &FClientTransferStatisticsModel::OnTransferUpdatedFromThread);
-		TickHandle = FTSTicker::GetCoreTicker().AddTicker(FTickerDelegate::CreateRaw(this, &FClientTransferStatisticsModel::Tick));
+		if (INetworkMessagingExtension* Statistics = Private::ClientTransferStatisticsModel::GetMessagingStatistics())
+		{
+			Statistics->OnTransferUpdatedFromThread().RemoveAll(this);
+		}
+		
+		FTSTicker::GetCoreTicker().RemoveTicker(TickHandle);
 	}
-}
 
-UE::MultiUserServer::FClientTransferStatisticsModel::~FClientTransferStatisticsModel()
-{
-	if (INetworkMessagingExtension* Statistics = Private::ClientTransferStatisticsModel::GetMessagingStatistics())
+	void FClientTransferStatisticsModel::OnTransferUpdatedFromThread(FTransferStatistics TransferStatistics)
 	{
-		Statistics->OnTransferUpdatedFromThread().RemoveAll(this);
+		if (AreRelevantStats(TransferStatistics))
+		{
+			AsyncStatQueue.Enqueue(FTransferStatItem{ FDateTime::Now(), TransferStatistics });
+		}
 	}
-	
-	FTSTicker::GetCoreTicker().RemoveTicker(TickHandle);
-}
 
-void UE::MultiUserServer::FClientTransferStatisticsModel::OnTransferUpdatedFromThread(FTransferStatistics TransferStatistics)
-{
-	AsyncStatQueue.Enqueue(TransferStatistics);
-}
+	bool FClientTransferStatisticsModel::AreRelevantStats(const FTransferStatistics& TransferStatistics) const
+	{
+		return IsSentToClient(TransferStatistics) || IsReceivedFromClient(TransferStatistics);
+	}
 
-bool UE::MultiUserServer::FClientTransferStatisticsModel::Tick(float DeltaTime)
-{
-	bool bAnyElements = false;
-	while (const TOptional<FTransferStatistics> TransferStatistics = AsyncStatQueue.Dequeue())
+	bool FClientTransferStatisticsModel::IsSentToClient(const FTransferStatistics& Item) const
+	{
+		if (INetworkMessagingExtension* Statistics = Private::ClientTransferStatisticsModel::GetMessagingStatistics())
+		{
+			const FGuid ClientNodeId = Statistics->GetNodeIdFromAddress(ClientAddress);
+			return ClientNodeId == Item.DestinationId;
+		}
+		return false;
+	}
+
+	bool FClientTransferStatisticsModel::IsReceivedFromClient(const FTransferStatistics& Item) const
+	{
+		// TODO: UDP module is missing a way to find out
+		return false; 
+	}
+	
+	bool FClientTransferStatisticsModel::Tick(float DeltaTime)
 	{
-		bAnyElements = true;
+		constexpr int32 NumStats = static_cast<int32>(EConcertTransferStatistic::Count);
+		FConcertStatTypeFlags ConcertStatFlags(false, NumStats);
+
+		RemoveOldStatTimelines(ConcertStatFlags);
+		const bool bRemovedAnyGroups = RemoveOldGroupedStats();
 		
-		const TSharedPtr<FTransferStatistics> NewValue = MakeShared<FTransferStatistics>(*TransferStatistics);
-		auto Pos = Algo::UpperBound(Stats, NewValue, [](const TSharedPtr<FTransferStatistics>& Value, const TSharedPtr<FTransferStatistics>& Check)
+		const FDateTime Now = FDateTime::Now();
+		FConcertTransferSamplePoint SamplingThisTick[NumStats] = { FConcertTransferSamplePoint{ Now }, FConcertTransferSamplePoint{ Now } };
+		static_assert(NumStats == 2);
+		while (const TOptional<FTransferStatItem> TransferStatistics = AsyncStatQueue.Dequeue())
 		{
-			return Value->MessageId > Check->MessageId;
-		});
-		if (Stats.Num() > 0 && Pos < Stats.Num() && Stats[Pos] && Stats[Pos]->MessageId == NewValue->MessageId)
+			constexpr int32 StatAsInt = static_cast<int32>(EConcertTransferStatistic::SentToClient);
+			ConcertStatFlags[StatAsInt] = true;
+			
+			UpdateStatTimeline(TransferStatistics.GetValue(), SamplingThisTick[StatAsInt]);
+			UpdateGroupedStats(TransferStatistics.GetValue());
+		}
+
+		if (bRemovedAnyGroups || ConcertStatFlags.CountSetBits() != 0)
 		{
-			Stats[Pos] = NewValue;
+			for (int32 i = 0; i < NumStats; ++i)
+			{
+				const bool bHadStatThisUpdate = ConcertStatFlags[i]; 
+				if (bHadStatThisUpdate)
+				{
+					const EConcertTransferStatistic Stat = static_cast<EConcertTransferStatistic>(i);
+					TransferStatisticsTimelines[Stat].Add(SamplingThisTick[i]);
+					OnTimelineUpdatedDelegates[Stat].Broadcast();
+				}
+			}
+			OnGroupsUpdatedDelegate.Broadcast();
 		}
-		else
+
+		return true;
+	}
+
+	namespace Private::ClientTransferStatisticsModel
+	{
+		bool IsTooOld(const FDateTime& StatCreationTime, const FDateTime& Now)
+		{
+			const FTimespan RetainTime = FTimespan::FromSeconds(60.f);
+			return RetainTime < Now - StatCreationTime;
+		}
+	}
+	
+	void FClientTransferStatisticsModel::RemoveOldStatTimelines(FConcertStatTypeFlags& ModifiedTimelines)
+	{
+		const FDateTime Now = FDateTime::Now();
+		for (auto TimelineIt = TransferStatisticsTimelines.CreateIterator(); TimelineIt; ++TimelineIt)
+		{
+			TArray<FConcertTransferSamplePoint>& SamplePoints = TimelineIt->Value;
+			bool bModifiedAny = false;
+			for (auto SampleIt = SamplePoints.CreateIterator(); SampleIt; ++SampleIt)
+			{
+				if (Private::ClientTransferStatisticsModel::IsTooOld(SampleIt->LocalTime, Now))
+				{
+					SampleIt.RemoveCurrent();
+					bModifiedAny = true;
+				}
+			}
+
+			const int32 Index = static_cast<int32>(TimelineIt->Key);
+			ModifiedTimelines[Index] = bModifiedAny;
+		}
+
+		for (auto IncompleteIt = IncompleteStatsUntilNow.CreateIterator(); IncompleteIt; ++IncompleteIt)
+		{
+			if (Private::ClientTransferStatisticsModel::IsTooOld(IncompleteIt->Value.LastUpdate, Now))
+			{
+				IncompleteIt.RemoveCurrent();
+			}
+		}
+	}
+
+	bool FClientTransferStatisticsModel::RemoveOldGroupedStats()
+	{
+		bool bRemovedAny = false;
+		const FDateTime Now = FDateTime::Now();
+		for (auto It = TransferStatisticsGroupedById.CreateIterator(); It; ++It)
 		{
-			Stats.Insert(NewValue, Pos);
+			const TSharedPtr<FTransferStatistics>& Stat = *It;
+			if (Private::ClientTransferStatisticsModel::IsTooOld(LastUpdateGroupUpdates[Stat->MessageId], Now))
+			{
+				LastUpdateGroupUpdates.Remove(Stat->MessageId);
+				It.RemoveCurrent();
+				bRemovedAny = true;
+			}
 		}
+		return bRemovedAny;
 	}
 
-	if (bAnyElements)
+	void FClientTransferStatisticsModel::UpdateStatTimeline(const FTransferStatItem& TransferStatistics, FConcertTransferSamplePoint& SamplingThisTick)
 	{
-		OnUpdatedDelegate.Broadcast();
+		uint64 BytesSentSinceLastTime = TransferStatistics.BytesSent;
+		// Same message ID may have been queued multiple times
+		if (FIncompleteMessageData* ItemThisTick = IncompleteStatsUntilNow.Find(TransferStatistics.MessageId))
+		{
+			BytesSentSinceLastTime -= ItemThisTick->BytesTransferredSoFar;
+		}
+
+		IncompleteStatsUntilNow.Add(TransferStatistics.MessageId, { TransferStatistics.BytesSent, FDateTime::Now() });
+		SamplingThisTick.BytesTransferred += BytesSentSinceLastTime;
 	}
 
-	return true;
+	void FClientTransferStatisticsModel::UpdateGroupedStats(const FTransferStatItem& TransferStatistics)
+	{
+		const TSharedPtr<FTransferStatistics> NewValue = MakeShared<FTransferStatistics>(TransferStatistics);
+		const auto Pos = Algo::LowerBound(TransferStatisticsGroupedById, NewValue, [](const TSharedPtr<FTransferStatistics>& Value, const TSharedPtr<FTransferStatistics>& Check)
+		{
+			return Value->MessageId < Check->MessageId;
+		});
+		if (TransferStatisticsGroupedById.Num() > 0 && Pos < TransferStatisticsGroupedById.Num() && TransferStatisticsGroupedById[Pos] && TransferStatisticsGroupedById[Pos]->MessageId == NewValue->MessageId)
+		{
+			*TransferStatisticsGroupedById[Pos] = *NewValue;
+		}
+		else
+		{
+			TransferStatisticsGroupedById.Insert(NewValue, Pos);
+		}
+
+		LastUpdateGroupUpdates.Add(NewValue->MessageId, FDateTime::Now());
+	}
 }
diff --git a/Engine/Plugins/Developer/Concert/ConcertApp/MultiUserServer/Source/MultiUserServer/Private/Widgets/Clients/Browser/Models/ClientTransferStatisticsModel.h b/Engine/Plugins/Developer/Concert/ConcertApp/MultiUserServer/Source/MultiUserServer/Private/Widgets/Clients/Browser/Models/ClientTransferStatisticsModel.h
--- a/Engine/Plugins/Developer/Concert/ConcertApp/MultiUserServer/Source/MultiUserServer/Private/Widgets/Clients/Browser/Models/ClientTransferStatisticsModel.h
+++ b/Engine/Plugins/Developer/Concert/ConcertApp/MultiUserServer/Source/MultiUserServer/Private/Widgets/Clients/Browser/Models/ClientTransferStatisticsModel.h
@@ -4,33 +4,72 @@
 
 #include "CoreMinimal.h"
 #include "IClientTransferStatisticsModel.h"
+#include "IMessageContext.h"
 #include "Containers/SpscQueue.h"
 #include "Containers/Ticker.h"
 
 namespace UE::MultiUserServer
 {
+	struct FTransferStatItem : FTransferStatistics
+	{
+		FDateTime LocalTime;
+
+		FTransferStatItem(const FDateTime& LocalTime, const FTransferStatistics& Data)
+			: FTransferStatistics(Data)
+			, LocalTime(LocalTime)
+		{}
+	};
+	
 	class FClientTransferStatisticsModel : public IClientTransferStatisticsModel
 	{
 	public:
 		
-		FClientTransferStatisticsModel(const FMessageAddress& MessageAddress);
+		FClientTransferStatisticsModel(const FMessageAddress& ClientAddress);
 		virtual ~FClientTransferStatisticsModel() override;
-
-		virtual const TArray<TSharedPtr<FTransferStatistics>>& GetSortedTransferStatistics() const override { return Stats; }
-		virtual FOnTransferStatisticsUpdated& OnTransferStatisticsUpdated() override { return OnUpdatedDelegate; }
+		
+		virtual const TArray<FConcertTransferSamplePoint>& GetTransferStatTimeline(EConcertTransferStatistic StatisticType) const override { return TransferStatisticsTimelines[StatisticType]; }
+		virtual const TArray<TSharedPtr<FTransferStatistics>>& GetTransferStatsGroupedById() const override { return TransferStatisticsGroupedById; }
+		virtual FOnTransferTimelineUpdated& OnTransferTimelineUpdated(EConcertTransferStatistic StatisticType) override { return OnTimelineUpdatedDelegates[StatisticType]; }
+		virtual FOnTransferGroupsUpdated& OnTransferGroupsUpdated() override { return OnGroupsUpdatedDelegate; }
 
 	private:
 
-		const FMessageAddress MessageAddress;
+		using FConcertStatTypeFlags = TBitArray<TInlineAllocator<static_cast<int32>(EConcertTransferStatistic::Count)>>;
+		using FMessageId = decltype(FTransferStatistics::MessageId);
+		struct FIncompleteMessageData
+		{
+			uint64 BytesTransferredSoFar;
+			FDateTime LastUpdate;
+		};
 
-		TSpscQueue<FTransferStatistics> AsyncStatQueue;
-		TArray<TSharedPtr<FTransferStatistics>> Stats;
+		const FMessageAddress ClientAddress;
+		TSpscQueue<FTransferStatItem> AsyncStatQueue;
+
+		// Timelines
+		TMap<EConcertTransferStatistic, TArray<FConcertTransferSamplePoint>> TransferStatisticsTimelines;
+		TMap<FMessageId, FIncompleteMessageData> IncompleteStatsUntilNow;
+
+		// Grouped
+		TArray<TSharedPtr<FTransferStatistics>> TransferStatisticsGroupedById;
+		TMap<FMessageId, FDateTime> LastUpdateGroupUpdates;
 		
-		FOnTransferStatisticsUpdated OnUpdatedDelegate;
+		// Delegate info
+		TMap<EConcertTransferStatistic, FOnTransferTimelineUpdated> OnTimelineUpdatedDelegates;
+		FOnTransferGroupsUpdated OnGroupsUpdatedDelegate;
 		FTSTicker::FDelegateHandle TickHandle;
 		
 		void OnTransferUpdatedFromThread(FTransferStatistics TransferStatistics);
+		bool AreRelevantStats(const FTransferStatistics& TransferStatistics) const;
+		bool IsSentToClient(const FTransferStatistics& Item) const;
+		bool IsReceivedFromClient(const FTransferStatistics& Item) const;
+		
 		bool Tick(float DeltaTime);
+		
+		void RemoveOldStatTimelines(FConcertStatTypeFlags& ModifiedTimelines);
+		bool RemoveOldGroupedStats();
+		
+		void UpdateStatTimeline(const FTransferStatItem& TransferStatistics, FConcertTransferSamplePoint& SamplingThisTick);
+		void UpdateGroupedStats(const FTransferStatItem& TransferStatistics);
 	};
 }
 
diff --git a/Engine/Plugins/Developer/Concert/ConcertApp/MultiUserServer/Source/MultiUserServer/Private/Widgets/Clients/Browser/Models/IClientTransferStatisticsModel.h b/Engine/Plugins/Developer/Concert/ConcertApp/MultiUserServer/Source/MultiUserServer/Private/Widgets/Clients/Browser/Models/IClientTransferStatisticsModel.h
--- a/Engine/Plugins/Developer/Concert/ConcertApp/MultiUserServer/Source/MultiUserServer/Private/Widgets/Clients/Browser/Models/IClientTransferStatisticsModel.h
+++ b/Engine/Plugins/Developer/Concert/ConcertApp/MultiUserServer/Source/MultiUserServer/Private/Widgets/Clients/Browser/Models/IClientTransferStatisticsModel.h
@@ -3,25 +3,52 @@
 #pragma once
 
 #include "CoreMinimal.h"
-#include "IMessageContext.h"
+#include "INetworkMessagingExtension.h"
 
 struct FMessageAddress;
-struct FTransferStatistics;
 
 namespace UE::MultiUserServer
 {
-	DECLARE_MULTICAST_DELEGATE(FOnTransferStatisticsUpdated);
+	enum class EConcertTransferStatistic
+	{
+		SentToClient,
+		ReceivedFromClient,
+
+		Count
+	};
+
+	DECLARE_MULTICAST_DELEGATE(FOnTransferTimelineUpdated);
+	DECLARE_MULTICAST_DELEGATE(FOnTransferGroupsUpdated);
+	
+	struct FConcertTransferSamplePoint
+	{
+		/** The time these stats were updated */
+		FDateTime LocalTime;
+
+		/** The number of bytes transferred, e.g. sent, received, etc. */
+		uint64 BytesTransferred;
+
+		FConcertTransferSamplePoint() = default;
+		explicit FConcertTransferSamplePoint(const FDateTime& Time, uint64 BytesTransferred = 0)
+			: LocalTime(Time)
+			, BytesTransferred(BytesTransferred)
+		{}
+	};
 	
-	/** Keeps track of a single client's transferstatistics */
+	/** Keeps track of a single client's transfer statistics */
 	class IClientTransferStatisticsModel
 	{
 	public:
 
-		/** Gets the transfer statistics sorted descending by message ID */
-		virtual const TArray<TSharedPtr<FTransferStatistics>>& GetSortedTransferStatistics() const = 0;
+		/** @return Contains the transfer statistics over time. Used for graphs. */
+		virtual const TArray<FConcertTransferSamplePoint>& GetTransferStatTimeline(EConcertTransferStatistic StatisticType) const = 0;
+		/** @return Sorted descending by message ID. Every message ID is unique. Used for list views. */
+		virtual const TArray<TSharedPtr<FTransferStatistics>>& GetTransferStatsGroupedById() const = 0;
 
-		/** Called when the transfer statistics change */
-		virtual FOnTransferStatisticsUpdated& OnTransferStatisticsUpdated() = 0;
+		/** Called when GetTransferStatisticsGroupedById changes */
+		virtual FOnTransferTimelineUpdated& OnTransferTimelineUpdated(EConcertTransferStatistic StatisticType) = 0;
+		/** Called when GetTransferStatisticsGroupedById changes */
+		virtual FOnTransferGroupsUpdated& OnTransferGroupsUpdated() = 0;
 
 		virtual ~IClientTransferStatisticsModel() = default;
 	};
diff --git a/Engine/Plugins/Developer/Concert/ConcertApp/MultiUserServer/Source/MultiUserServer/Private/Widgets/Clients/Browser/SClientTransferStatTable.cpp b/Engine/Plugins/Developer/Concert/ConcertApp/MultiUserServer/Source/MultiUserServer/Private/Widgets/Clients/Browser/SClientTransferStatTable.cpp
--- a/Engine/Plugins/Developer/Concert/ConcertApp/MultiUserServer/Source/MultiUserServer/Private/Widgets/Clients/Browser/SClientTransferStatTable.cpp
+++ b/Engine/Plugins/Developer/Concert/ConcertApp/MultiUserServer/Source/MultiUserServer/Private/Widgets/Clients/Browser/SClientTransferStatTable.cpp
@@ -54,7 +54,7 @@ namespace UE::MultiUserServer
 					.AutoWidth()
 					[
 						SNew(STextBlock)
-						.Text_Lambda([this](){ return FText::AsNumber(Stats->MessageId); })
+						.Text(FText::AsNumber(Stats->MessageId))
 					];
 			}
 			if (HeaderIdName_SentSegments == ColumnName)
@@ -90,7 +90,7 @@ namespace UE::MultiUserServer
 					.AutoWidth()
 					[
 						SNew(STextBlock)
-						.Text_Lambda([this](){ return FText::AsNumber(Stats->BytesToSend); })
+						.Text(FText::AsNumber(Stats->BytesToSend))
 					];
 			}
 			return SNullWidget::NullWidget;
@@ -101,14 +101,14 @@ namespace UE::MultiUserServer
 	{
 		StatsModel = InStatsModel;
 
-		StatsModel->OnTransferStatisticsUpdated().AddSP(this, &SClientTransferStatTable::OnTransferStatisticsUpdated);
+		StatsModel->OnTransferGroupsUpdated().AddSP(this, &SClientTransferStatTable::OnTransferStatisticsUpdated);
 
 		ChildSlot
 		[
 			SAssignNew(SegmenterListView, SListView<TSharedPtr<FTransferStatistics>>)
-			.ListItemsSource(&InStatsModel->GetSortedTransferStatistics())
+			.ListItemsSource(&InStatsModel->GetTransferStatsGroupedById())
 			.OnGenerateRow(this, &SClientTransferStatTable::OnGenerateActivityRowWidget)
-			.SelectionMode(ESelectionMode::None)
+			.SelectionMode(ESelectionMode::Multi)
 			.HeaderRow
 			(
 				SNew(SHeaderRow)
diff --git a/Engine/Plugins/Developer/Concert/ConcertApp/MultiUserServer/Source/MultiUserServer/Private/Widgets/Clients/Browser/SConcertClientBrowserItem.cpp b/Engine/Plugins/Developer/Concert/ConcertApp/MultiUserServer/Source/MultiUserServer/Private/Widgets/Clients/Browser/SConcertClientBrowserItem.cpp
--- a/Engine/Plugins/Developer/Concert/ConcertApp/MultiUserServer/Source/MultiUserServer/Private/Widgets/Clients/Browser/SConcertClientBrowserItem.cpp
+++ b/Engine/Plugins/Developer/Concert/ConcertApp/MultiUserServer/Source/MultiUserServer/Private/Widgets/Clients/Browser/SConcertClientBrowserItem.cpp
@@ -7,6 +7,7 @@
 #include "INetworkMessagingExtension.h"
 #include "SClientNetworkStats.h"
 #include "SClientTransferStatTable.h"
+#include "Graph/SClientNetworkGraphs.h"
 #include "Models/ClientTransferStatisticsModel.h"
 #include "Models/IClientNetworkStatisticsModel.h"
 #include "Styling/StyleColors.h"
@@ -44,7 +45,7 @@ void UE::MultiUserServer::SConcertClientBrowserItem::Construct(const FArguments&
 				// Inside of thumbnail is darker
 				SNew(SBorder)
 				.BorderImage(FConcertServerStyle::Get().GetBrush("Concert.Clients.ThumbnailTitle"))
-				.Padding(4.f)
+				.Padding(2.f)
 				[
 					SNew(SVerticalBox)
 
@@ -57,7 +58,7 @@ void UE::MultiUserServer::SConcertClientBrowserItem::Construct(const FArguments&
 					
 					+SVerticalBox::Slot()
 					.FillHeight(1.f)
-					.Padding(0.f, 5.f, 0.f, 0.f)
+					.Padding(-2.f, 5.f, -2.f, 0.f)
 					[
 						CreateContentArea()
 					]
@@ -73,7 +74,7 @@ void UE::MultiUserServer::SConcertClientBrowserItem::Construct(const FArguments&
 
 					+SVerticalBox::Slot()
 					.AutoHeight()
-					.Padding(-4.f, 5.f, -4.f, -4.f)
+					.Padding(-2.f, 5.f, -2.f, -2.f)
 					.VAlign(VAlign_Bottom)
 					[
 						CreateFooter()
@@ -137,10 +138,10 @@ TSharedRef<SWidget> UE::MultiUserServer::SConcertClientBrowserItem::CreateConten
 	return SNew(SOverlay)
 		+SOverlay::Slot()
 		[
-			// TODO: Graph
-			SNullWidget::NullWidget
+			SNew(SClientNetworkGraphs, TransferStatsModel.ToSharedRef())
+			.Visibility_Lambda([this](){ return GetDisplayMode() == EClientDisplayMode::NetworkGraph ? EVisibility::Visible : EVisibility::Collapsed; })
 		]
-		+SOverlay::Slot()
+		+SOverlay::Slot().Padding(4.f, 0.f)
 		[
 			SNew(SClientTransferStatTable, TransferStatsModel.ToSharedRef())
 			.Visibility_Lambda([this](){ return GetDisplayMode() == EClientDisplayMode::SegementTable ? EVisibility::Visible : EVisibility::Collapsed; })
