Warn users that muting and unmuting can cause crashes

diff --git a/Engine/Plugins/Developer/Concert/ConcertApp/MultiUserServer/Source/MultiUserServer/Private/Settings/MultiUserServerUserPreferences.cpp b/Engine/Plugins/Developer/Concert/ConcertApp/MultiUserServer/Source/MultiUserServer/Private/Settings/MultiUserServerUserPreferences.cpp
new file mode 100644
--- /dev/null
+++ b/Engine/Plugins/Developer/Concert/ConcertApp/MultiUserServer/Source/MultiUserServer/Private/Settings/MultiUserServerUserPreferences.cpp
@@ -0,0 +1,8 @@
+// Copyright Epic Games, Inc. All Rights Reserved.
+
+#include "MultiUserServerUserPreferences.h"
+
+UMultiUserServerUserPreferences* UMultiUserServerUserPreferences::GetSettings()
+{
+	return GetMutableDefault<UMultiUserServerUserPreferences>();
+}
diff --git a/Engine/Plugins/Developer/Concert/ConcertApp/MultiUserServer/Source/MultiUserServer/Private/Settings/MultiUserServerUserPreferences.h b/Engine/Plugins/Developer/Concert/ConcertApp/MultiUserServer/Source/MultiUserServer/Private/Settings/MultiUserServerUserPreferences.h
new file mode 100644
--- /dev/null
+++ b/Engine/Plugins/Developer/Concert/ConcertApp/MultiUserServer/Source/MultiUserServer/Private/Settings/MultiUserServerUserPreferences.h
@@ -0,0 +1,25 @@
+// Copyright Epic Games, Inc. All Rights Reserved.
+
+#pragma once
+
+#include "CoreMinimal.h"
+#include "UObject/Object.h"
+#include "MultiUserServerUserPreferences.generated.h"
+
+/**
+ * 
+ */
+UCLASS(Config = MultiUserServerUserSettings, DefaultConfig)
+class MULTIUSERSERVER_API UMultiUserServerUserPreferences : public UObject
+{
+	GENERATED_BODY()
+public:
+	
+	static UMultiUserServerUserPreferences* GetSettings();
+
+	UPROPERTY(Config)
+	bool bWarnUserAboutMuting = true;
+
+	UPROPERTY(Config)
+	bool bWarnUserAboutUnmuting = true;
+};
diff --git a/Engine/Plugins/Developer/Concert/ConcertApp/MultiUserServer/Source/MultiUserServer/Private/Widgets/SessionTabs/Archived/ArchivedConcertSessionTab.cpp b/Engine/Plugins/Developer/Concert/ConcertApp/MultiUserServer/Source/MultiUserServer/Private/Widgets/SessionTabs/Archived/ArchivedConcertSessionTab.cpp
--- a/Engine/Plugins/Developer/Concert/ConcertApp/MultiUserServer/Source/MultiUserServer/Private/Widgets/SessionTabs/Archived/ArchivedConcertSessionTab.cpp
+++ b/Engine/Plugins/Developer/Concert/ConcertApp/MultiUserServer/Source/MultiUserServer/Private/Widgets/SessionTabs/Archived/ArchivedConcertSessionTab.cpp
@@ -19,7 +19,9 @@
 #include "Algo/AllOf.h"
 #include "Algo/Transform.h"
 #include "Dialog/SMessageDialog.h"
+#include "Settings/MultiUserServerUserPreferences.h"
 #include "Widgets/Docking/SDockTab.h"
+#include "Widgets/Util/SDoNotShowAgainDialog.h"
 
 #define LOCTEXT_NAMESPACE "UnrealMultiUserUI.FArchivedConcertSessionTab"
 
@@ -117,8 +119,44 @@ FCanPerformActionResult FArchivedConcertSessionTab::CanDeleteActivities(const TS
 
 void FArchivedConcertSessionTab::OnRequestMuteActivities(const TSet<TSharedRef<FConcertSessionActivity>>& ActivitiesToMute) const
 {
-	constexpr bool bShouldMute = true;
-	SpawnDialogForMutingOrUnmuting(bShouldMute, ActivitiesToMute);
+	using namespace UE::MultiUserServer;
+	
+	if (UMultiUserServerUserPreferences::GetSettings()->bWarnUserAboutMuting)
+	{
+		const TSharedRef<SDoNotShowAgainDialog> Dialog =
+		SNew(SDoNotShowAgainDialog)
+			.Title(LOCTEXT("MuteTitle", "Caution: Muting"))
+			.Buttons({
+				SCustomDialog::FButton(LOCTEXT("Continue", "Continue"))
+					.SetPrimary(true)
+					.SetOnClicked(FSimpleDelegate::CreateLambda([this, ActivitiesToMute]()
+					{
+						constexpr bool bShouldMute = true;
+						SpawnDialogForMutingOrUnmuting(bShouldMute, ActivitiesToMute);
+					})),
+				SCustomDialog::FButton(LOCTEXT("Cancel", "Cancel"))
+			})
+			.DoNotShowAgainCallback(SDoNotShowAgainDialog::FOnClosed::CreateLambda([](bool bDoNotShowAgain)
+			{
+				UMultiUserServerUserPreferences* UserPreferences = UMultiUserServerUserPreferences::GetSettings();
+				UserPreferences->bWarnUserAboutMuting = !bDoNotShowAgain;
+				UserPreferences->SaveConfig();
+			}))
+			.Content()
+			[
+				SNew(STextBlock)
+				.Text(LOCTEXT("MuteWarning", "Muted activities are not sent to clients. Activities you select are analysed for dependencies.\nExample: Muting the creation of a map also mutes editing an actor in that level.\n\nMuting is not always safe.\n\nDependency analysis cannot detect all cases:\nsometimes transactions may depend on other transaction activities having occured before.\nThis cannot be detected."))
+			];
+
+		FConcertServerUIModule::Get()
+			.GetModalWindowManager()
+			->ShowFakeModalWindow(Dialog);
+	}
+	else
+	{
+		constexpr bool bShouldMute = true;
+		SpawnDialogForMutingOrUnmuting(bShouldMute, ActivitiesToMute);
+	}	
 }
 
 FCanPerformActionResult FArchivedConcertSessionTab::CanMuteActivities(const TSet<TSharedRef<FConcertSessionActivity>>& ActivitiesToDelete) const
@@ -146,8 +184,44 @@ FCanPerformActionResult FArchivedConcertSessionTab::CanMuteActivities(const TSet
 
 void FArchivedConcertSessionTab::OnRequestUnmuteActivities(const TSet<TSharedRef<FConcertSessionActivity>>& ActivitiesToUnmute) const
 {
-	constexpr bool bShouldMute = false;
-	SpawnDialogForMutingOrUnmuting(bShouldMute, ActivitiesToUnmute);
+	using namespace UE::MultiUserServer;
+	
+	if (UMultiUserServerUserPreferences::GetSettings()->bWarnUserAboutUnmuting)
+	{
+		const TSharedRef<SDoNotShowAgainDialog> Dialog =
+		SNew(SDoNotShowAgainDialog)
+			.Title(LOCTEXT("UnmuteTitle", "Caution: Unmuting"))
+			.Buttons({
+				SCustomDialog::FButton(LOCTEXT("Continue", "Continue"))
+					.SetPrimary(true)
+					.SetOnClicked(FSimpleDelegate::CreateLambda([this, ActivitiesToUnmute]()
+					{
+						constexpr bool bShouldMute = false;
+						SpawnDialogForMutingOrUnmuting(bShouldMute, ActivitiesToUnmute);
+					})),
+				SCustomDialog::FButton(LOCTEXT("Cancel", "Cancel"))
+			})
+			.DoNotShowAgainCallback(SDoNotShowAgainDialog::FOnClosed::CreateLambda([](bool bDoNotShowAgain)
+			{
+				UMultiUserServerUserPreferences* UserPreferences = UMultiUserServerUserPreferences::GetSettings();
+				UserPreferences->bWarnUserAboutUnmuting = !bDoNotShowAgain;
+				UserPreferences->SaveConfig();
+			}))
+			.Content()
+			[
+				SNew(STextBlock)
+				.Text(LOCTEXT("UnmuteWarning", "Unmuting is not always safe\nPay attention if this session had muted activities, was unarchived, and now you're unmuting said activities again.\n\nExample of conflict:\n1. Create package Foo\nNow mute 1 and unarchive session.\n1. Create package Foo (muted)\n2. Create package Foo\nArchive the session, and try to unmute 1. You'd end up with a history that creates Foo twice, which is not valid.\n\nSimilar invalid cases can happen with transaction activities."))
+			];
+
+		FConcertServerUIModule::Get()
+			.GetModalWindowManager()
+			->ShowFakeModalWindow(Dialog);
+	}
+	else
+	{
+		constexpr bool bShouldMute = false;
+		SpawnDialogForMutingOrUnmuting(bShouldMute, ActivitiesToUnmute);
+	}	
 }
 
 FCanPerformActionResult FArchivedConcertSessionTab::CanUnmuteActivities(const TSet<TSharedRef<FConcertSessionActivity>>& ActivitiesToDelete) const
diff --git a/Engine/Plugins/Developer/Concert/ConcertApp/MultiUserServer/Source/MultiUserServer/Private/Widgets/Util/SDoNotShowAgainDialog.cpp b/Engine/Plugins/Developer/Concert/ConcertApp/MultiUserServer/Source/MultiUserServer/Private/Widgets/Util/SDoNotShowAgainDialog.cpp
new file mode 100644
--- /dev/null
+++ b/Engine/Plugins/Developer/Concert/ConcertApp/MultiUserServer/Source/MultiUserServer/Private/Widgets/Util/SDoNotShowAgainDialog.cpp
@@ -0,0 +1,59 @@
+// Copyright Epic Games, Inc. All Rights Reserved.
+
+#include "SDoNotShowAgainDialog.h"
+
+#include "Widgets/Input/SCheckBox.h"
+#include "Widgets/SBoxPanel.h"
+#include "Widgets/Text/STextBlock.h"
+
+#define LOCTEXT_NAMESPACE "UnrealMultiUserUI.SDoNotShowAgainDialog"
+
+void UE::MultiUserServer::SDoNotShowAgainDialog::Construct(const FArguments& InArgs)
+{
+	SCustomDialog::Construct(
+		SCustomDialog::FArguments()
+		.Title(InArgs._Title)
+		.Buttons(InArgs._Buttons)
+		.Content()
+		[
+			SNew(SVerticalBox)
+			
+			+SVerticalBox::Slot()
+			.HAlign(HAlign_Fill)
+			.VAlign(VAlign_Top)
+			.FillHeight(1.f)
+			.Padding(0, 5, 0, 10)
+			[
+				InArgs._Content.Widget
+			]
+
+			+SVerticalBox::Slot()
+			.AutoHeight()
+			.HAlign(HAlign_Left)
+			.Padding(0, 5, 0, 10)
+			[
+				SNew(SHorizontalBox)
+				
+				+SHorizontalBox::Slot()
+				.AutoWidth()
+				[
+					SAssignNew(DoNotShowAgainCheckbox, SCheckBox)
+				]
+				
+				+SHorizontalBox::Slot()
+				.AutoWidth()
+				[
+					SNew(STextBlock)
+					.Text(LOCTEXT("DoNotShowAgain", "Do not show again"))
+				]
+				
+			]
+		]
+		.OnClosed_Lambda([this, Callback = InArgs._DoNotShowAgainCallback]()
+		{
+			Callback.Execute(DoNotShowAgainCheckbox->IsChecked());
+		})
+	);
+}
+
+#undef LOCTEXT_NAMESPACE
diff --git a/Engine/Plugins/Developer/Concert/ConcertApp/MultiUserServer/Source/MultiUserServer/Private/Widgets/Util/SDoNotShowAgainDialog.h b/Engine/Plugins/Developer/Concert/ConcertApp/MultiUserServer/Source/MultiUserServer/Private/Widgets/Util/SDoNotShowAgainDialog.h
new file mode 100644
--- /dev/null
+++ b/Engine/Plugins/Developer/Concert/ConcertApp/MultiUserServer/Source/MultiUserServer/Private/Widgets/Util/SDoNotShowAgainDialog.h
@@ -0,0 +1,41 @@
+// Copyright Epic Games, Inc. All Rights Reserved.
+
+#pragma once
+
+#include "CoreMinimal.h"
+#include "Dialog/SCustomDialog.h"
+#include "Widgets/DeclarativeSyntaxSupport.h"
+
+class FArchivedSessionHistoryController;
+class IConcertSyncServer;
+class SCheckBox;
+
+namespace UE::ConcertSyncCore
+{
+	struct FHistoryAnalysisResult;
+}
+
+namespace UE::MultiUserServer
+{
+	/** Text message with a "Do not show again" checkbox. */
+	class SDoNotShowAgainDialog : public SCustomDialog
+	{
+	public:
+
+		DECLARE_DELEGATE_OneParam(FOnClosed, bool /*bDoNotShowAgain*/);
+
+		SLATE_BEGIN_ARGS(SDoNotShowAgainDialog)
+		{}
+			SLATE_ARGUMENT(FText, Title)
+			SLATE_ARGUMENT(TArray<FButton>, Buttons)
+			SLATE_NAMED_SLOT(FArguments, Content)
+			SLATE_EVENT(FOnClosed, DoNotShowAgainCallback)
+		SLATE_END_ARGS()
+
+		void Construct(const FArguments& InArgs);
+
+	private:
+
+		TSharedPtr<SCheckBox> DoNotShowAgainCheckbox;
+	};
+}
