Add data reporting mechanism for inbound messages.

diff --git a/Engine/Plugins/Developer/Concert/ConcertApp/MultiUserServer/Source/MultiUserServer/Private/Widgets/Clients/Browser/Models/ClientTransferStatisticsModel.cpp b/Engine/Plugins/Developer/Concert/ConcertApp/MultiUserServer/Source/MultiUserServer/Private/Widgets/Clients/Browser/Models/ClientTransferStatisticsModel.cpp
--- a/Engine/Plugins/Developer/Concert/ConcertApp/MultiUserServer/Source/MultiUserServer/Private/Widgets/Clients/Browser/Models/ClientTransferStatisticsModel.cpp
+++ b/Engine/Plugins/Developer/Concert/ConcertApp/MultiUserServer/Source/MultiUserServer/Private/Widgets/Clients/Browser/Models/ClientTransferStatisticsModel.cpp
@@ -51,7 +51,7 @@ namespace UE::MultiUserServer
 		
 		if (INetworkMessagingExtension* Statistics = Private::ClientTransferStatisticsModel::GetMessagingStatistics())
 		{
-			Statistics->OnTransferUpdatedFromThread().AddRaw(this, &FClientTransferStatisticsModel::OnTransferUpdatedFromThread);
+			Statistics->OnOutboundTransferUpdatedFromThread().AddRaw(this, &FClientTransferStatisticsModel::OnTransferUpdatedFromThread);
 			TickHandle = FTSTicker::GetCoreTicker().AddTicker(FTickerDelegate::CreateRaw(this, &FClientTransferStatisticsModel::Tick));
 		}
 	}
@@ -60,13 +60,13 @@ namespace UE::MultiUserServer
 	{
 		if (INetworkMessagingExtension* Statistics = Private::ClientTransferStatisticsModel::GetMessagingStatistics())
 		{
-			Statistics->OnTransferUpdatedFromThread().RemoveAll(this);
+			Statistics->OnOutboundTransferUpdatedFromThread().RemoveAll(this);
 		}
 		
 		FTSTicker::GetCoreTicker().RemoveTicker(TickHandle);
 	}
 
-	void FClientTransferStatisticsModel::OnTransferUpdatedFromThread(FTransferStatistics TransferStatistics)
+	void FClientTransferStatisticsModel::OnTransferUpdatedFromThread(FOutboundTransferStatistics TransferStatistics)
 	{
 		if (AreRelevantStats(TransferStatistics))
 		{
@@ -74,12 +74,12 @@ namespace UE::MultiUserServer
 		}
 	}
 
-	bool FClientTransferStatisticsModel::AreRelevantStats(const FTransferStatistics& TransferStatistics) const
+	bool FClientTransferStatisticsModel::AreRelevantStats(const FOutboundTransferStatistics& TransferStatistics) const
 	{
 		return IsSentToClient(TransferStatistics) || IsReceivedFromClient(TransferStatistics);
 	}
 
-	bool FClientTransferStatisticsModel::IsSentToClient(const FTransferStatistics& Item) const
+	bool FClientTransferStatisticsModel::IsSentToClient(const FOutboundTransferStatistics& Item) const
 	{
 		if (INetworkMessagingExtension* Statistics = Private::ClientTransferStatisticsModel::GetMessagingStatistics())
 		{
@@ -89,7 +89,7 @@ namespace UE::MultiUserServer
 		return false;
 	}
 
-	bool FClientTransferStatisticsModel::IsReceivedFromClient(const FTransferStatistics& Item) const
+	bool FClientTransferStatisticsModel::IsReceivedFromClient(const FOutboundTransferStatistics& Item) const
 	{
 		// TODO: UDP module is missing a way to find out
 		return false; 
@@ -177,7 +177,7 @@ namespace UE::MultiUserServer
 		const FDateTime Now = FDateTime::Now();
 		for (auto It = TransferStatisticsGroupedById.CreateIterator(); It; ++It)
 		{
-			const TSharedPtr<FTransferStatistics>& Stat = *It;
+			const TSharedPtr<FOutboundTransferStatistics>& Stat = *It;
 			if (Private::ClientTransferStatisticsModel::IsTooOld(LastUpdateGroupUpdates[Stat->MessageId], Now))
 			{
 				LastUpdateGroupUpdates.Remove(Stat->MessageId);
@@ -203,8 +203,8 @@ namespace UE::MultiUserServer
 
 	void FClientTransferStatisticsModel::UpdateGroupedStats(const FTransferStatItem& TransferStatistics)
 	{
-		const TSharedPtr<FTransferStatistics> NewValue = MakeShared<FTransferStatistics>(TransferStatistics);
-		const auto Pos = Algo::LowerBound(TransferStatisticsGroupedById, NewValue, [](const TSharedPtr<FTransferStatistics>& Value, const TSharedPtr<FTransferStatistics>& Check)
+		const TSharedPtr<FOutboundTransferStatistics> NewValue = MakeShared<FOutboundTransferStatistics>(TransferStatistics);
+		const auto Pos = Algo::LowerBound(TransferStatisticsGroupedById, NewValue, [](const TSharedPtr<FOutboundTransferStatistics>& Value, const TSharedPtr<FOutboundTransferStatistics>& Check)
 		{
 			return Value->MessageId < Check->MessageId;
 		});
diff --git a/Engine/Plugins/Developer/Concert/ConcertApp/MultiUserServer/Source/MultiUserServer/Private/Widgets/Clients/Browser/Models/ClientTransferStatisticsModel.h b/Engine/Plugins/Developer/Concert/ConcertApp/MultiUserServer/Source/MultiUserServer/Private/Widgets/Clients/Browser/Models/ClientTransferStatisticsModel.h
--- a/Engine/Plugins/Developer/Concert/ConcertApp/MultiUserServer/Source/MultiUserServer/Private/Widgets/Clients/Browser/Models/ClientTransferStatisticsModel.h
+++ b/Engine/Plugins/Developer/Concert/ConcertApp/MultiUserServer/Source/MultiUserServer/Private/Widgets/Clients/Browser/Models/ClientTransferStatisticsModel.h
@@ -10,12 +10,12 @@
 
 namespace UE::MultiUserServer
 {
-	struct FTransferStatItem : FTransferStatistics
+	struct FTransferStatItem : FOutboundTransferStatistics
 	{
 		FDateTime LocalTime;
 
-		FTransferStatItem(const FDateTime& LocalTime, const FTransferStatistics& Data)
-			: FTransferStatistics(Data)
+		FTransferStatItem(const FDateTime& LocalTime, const FOutboundTransferStatistics& Data)
+			: FOutboundTransferStatistics(Data)
 			, LocalTime(LocalTime)
 		{}
 	};
@@ -28,14 +28,14 @@ namespace UE::MultiUserServer
 		virtual ~FClientTransferStatisticsModel() override;
 		
 		virtual const TArray<FConcertTransferSamplePoint>& GetTransferStatTimeline(EConcertTransferStatistic StatisticType) const override { return TransferStatisticsTimelines[StatisticType]; }
-		virtual const TArray<TSharedPtr<FTransferStatistics>>& GetTransferStatsGroupedById() const override { return TransferStatisticsGroupedById; }
+		virtual const TArray<TSharedPtr<FOutboundTransferStatistics>>& GetTransferStatsGroupedById() const override { return TransferStatisticsGroupedById; }
 		virtual FOnTransferTimelineUpdated& OnTransferTimelineUpdated(EConcertTransferStatistic StatisticType) override { return OnTimelineUpdatedDelegates[StatisticType]; }
 		virtual FOnTransferGroupsUpdated& OnTransferGroupsUpdated() override { return OnGroupsUpdatedDelegate; }
 
 	private:
 
 		using FConcertStatTypeFlags = TBitArray<TInlineAllocator<static_cast<int32>(EConcertTransferStatistic::Count)>>;
-		using FMessageId = decltype(FTransferStatistics::MessageId);
+		using FMessageId = decltype(FOutboundTransferStatistics::MessageId);
 		struct FIncompleteMessageData
 		{
 			uint64 BytesTransferredSoFar;
@@ -50,7 +50,7 @@ namespace UE::MultiUserServer
 		TMap<FMessageId, FIncompleteMessageData> IncompleteStatsUntilNow;
 
 		// Grouped
-		TArray<TSharedPtr<FTransferStatistics>> TransferStatisticsGroupedById;
+		TArray<TSharedPtr<FOutboundTransferStatistics>> TransferStatisticsGroupedById;
 		TMap<FMessageId, FDateTime> LastUpdateGroupUpdates;
 		
 		// Delegate info
@@ -58,10 +58,10 @@ namespace UE::MultiUserServer
 		FOnTransferGroupsUpdated OnGroupsUpdatedDelegate;
 		FTSTicker::FDelegateHandle TickHandle;
 		
-		void OnTransferUpdatedFromThread(FTransferStatistics TransferStatistics);
-		bool AreRelevantStats(const FTransferStatistics& TransferStatistics) const;
-		bool IsSentToClient(const FTransferStatistics& Item) const;
-		bool IsReceivedFromClient(const FTransferStatistics& Item) const;
+		void OnTransferUpdatedFromThread(FOutboundTransferStatistics TransferStatistics);
+		bool AreRelevantStats(const FOutboundTransferStatistics& TransferStatistics) const;
+		bool IsSentToClient(const FOutboundTransferStatistics& Item) const;
+		bool IsReceivedFromClient(const FOutboundTransferStatistics& Item) const;
 		
 		bool Tick(float DeltaTime);
 		
diff --git a/Engine/Plugins/Developer/Concert/ConcertApp/MultiUserServer/Source/MultiUserServer/Private/Widgets/Clients/Browser/Models/IClientTransferStatisticsModel.h b/Engine/Plugins/Developer/Concert/ConcertApp/MultiUserServer/Source/MultiUserServer/Private/Widgets/Clients/Browser/Models/IClientTransferStatisticsModel.h
--- a/Engine/Plugins/Developer/Concert/ConcertApp/MultiUserServer/Source/MultiUserServer/Private/Widgets/Clients/Browser/Models/IClientTransferStatisticsModel.h
+++ b/Engine/Plugins/Developer/Concert/ConcertApp/MultiUserServer/Source/MultiUserServer/Private/Widgets/Clients/Browser/Models/IClientTransferStatisticsModel.h
@@ -43,7 +43,7 @@ namespace UE::MultiUserServer
 		/** @return Contains the transfer statistics over time. Used for graphs. */
 		virtual const TArray<FConcertTransferSamplePoint>& GetTransferStatTimeline(EConcertTransferStatistic StatisticType) const = 0;
 		/** @return Sorted descending by message ID. Every message ID is unique. Used for list views. */
-		virtual const TArray<TSharedPtr<FTransferStatistics>>& GetTransferStatsGroupedById() const = 0;
+		virtual const TArray<TSharedPtr<FOutboundTransferStatistics>>& GetTransferStatsGroupedById() const = 0;
 
 		/** Called when GetTransferStatisticsGroupedById changes */
 		virtual FOnTransferTimelineUpdated& OnTransferTimelineUpdated(EConcertTransferStatistic StatisticType) = 0;
diff --git a/Engine/Plugins/Developer/Concert/ConcertApp/MultiUserServer/Source/MultiUserServer/Private/Widgets/Clients/Browser/SClientTransferStatTable.cpp b/Engine/Plugins/Developer/Concert/ConcertApp/MultiUserServer/Source/MultiUserServer/Private/Widgets/Clients/Browser/SClientTransferStatTable.cpp
--- a/Engine/Plugins/Developer/Concert/ConcertApp/MultiUserServer/Source/MultiUserServer/Private/Widgets/Clients/Browser/SClientTransferStatTable.cpp
+++ b/Engine/Plugins/Developer/Concert/ConcertApp/MultiUserServer/Source/MultiUserServer/Private/Widgets/Clients/Browser/SClientTransferStatTable.cpp
@@ -23,13 +23,13 @@ namespace UE::MultiUserServer
 		const FName HeaderIdName_DataRate     = TEXT("DataRate");
 	}
 
-	class SClientTransferStatTableRow : public SMultiColumnTableRow<TSharedRef<FTransferStatistics>>
+	class SClientTransferStatTableRow : public SMultiColumnTableRow<TSharedRef<FOutboundTransferStatistics>>
 	{
-		using Super = SMultiColumnTableRow<TSharedRef<FTransferStatistics>>;
+		using Super = SMultiColumnTableRow<TSharedRef<FOutboundTransferStatistics>>;
 	public:
 		
 		SLATE_BEGIN_ARGS(SClientTransferStatTableRow) { }
-			SLATE_ARGUMENT(TSharedPtr<FTransferStatistics>, Stats)
+			SLATE_ARGUMENT(TSharedPtr<FOutboundTransferStatistics>, Stats)
 		SLATE_END_ARGS()
 		
 		void Construct(const FArguments& InArgs, const TSharedRef<STableViewBase>& InOwerTableView)
@@ -40,7 +40,7 @@ namespace UE::MultiUserServer
 		
 	private:
 
-		TSharedPtr<FTransferStatistics> Stats;
+		TSharedPtr<FOutboundTransferStatistics> Stats;
 
 		virtual TSharedRef<SWidget> GenerateWidgetForColumn(const FName& ColumnName) override
 		{
@@ -105,7 +105,7 @@ namespace UE::MultiUserServer
 
 		ChildSlot
 		[
-			SAssignNew(SegmenterListView, SListView<TSharedPtr<FTransferStatistics>>)
+			SAssignNew(SegmenterListView, SListView<TSharedPtr<FOutboundTransferStatistics>>)
 			.ListItemsSource(&InStatsModel->GetTransferStatsGroupedById())
 			.OnGenerateRow(this, &SClientTransferStatTable::OnGenerateActivityRowWidget)
 			.SelectionMode(ESelectionMode::Multi)
@@ -128,7 +128,7 @@ namespace UE::MultiUserServer
 		];
 	}
 
-	TSharedRef<ITableRow> SClientTransferStatTable::OnGenerateActivityRowWidget(TSharedPtr<FTransferStatistics> InStats, const TSharedRef<STableViewBase>& OwnerTable) const
+	TSharedRef<ITableRow> SClientTransferStatTable::OnGenerateActivityRowWidget(TSharedPtr<FOutboundTransferStatistics> InStats, const TSharedRef<STableViewBase>& OwnerTable) const
 	{
 		return SNew(SClientTransferStatTableRow, OwnerTable)
 			.Stats(InStats);;
@@ -140,4 +140,4 @@ namespace UE::MultiUserServer
 	}
 }
 
-#undef LOCTEXT_NAMESPACE
\ No newline at end of file
+#undef LOCTEXT_NAMESPACE
diff --git a/Engine/Plugins/Developer/Concert/ConcertApp/MultiUserServer/Source/MultiUserServer/Private/Widgets/Clients/Browser/SClientTransferStatTable.h b/Engine/Plugins/Developer/Concert/ConcertApp/MultiUserServer/Source/MultiUserServer/Private/Widgets/Clients/Browser/SClientTransferStatTable.h
--- a/Engine/Plugins/Developer/Concert/ConcertApp/MultiUserServer/Source/MultiUserServer/Private/Widgets/Clients/Browser/SClientTransferStatTable.h
+++ b/Engine/Plugins/Developer/Concert/ConcertApp/MultiUserServer/Source/MultiUserServer/Private/Widgets/Clients/Browser/SClientTransferStatTable.h
@@ -8,14 +8,14 @@
 class ITableRow;
 template<typename T> class SListView;
 class STableViewBase;
-struct FTransferStatistics;
+struct FOutboundTransferStatistics;
 
 namespace UE::MultiUserServer
 {
 	class IClientTransferStatisticsModel;
 
 	/**
-	 * Displays FTransferStatistics as they are updated by INetworkMessagingExtension.
+	 * Displays FOutboundTransferStatistics as they are updated by INetworkMessagingExtension.
 	 */
 	class SClientTransferStatTable : public SCompoundWidget
 	{
@@ -32,9 +32,9 @@ namespace UE::MultiUserServer
 		/** Tells us when the transfer stats have changed */
 		TSharedPtr<IClientTransferStatisticsModel> StatsModel;
 		/** Displays the transfer stats */
-		TSharedPtr<SListView<TSharedPtr<FTransferStatistics>>> SegmenterListView;
+		TSharedPtr<SListView<TSharedPtr<FOutboundTransferStatistics>>> SegmenterListView;
 
-		TSharedRef<ITableRow> OnGenerateActivityRowWidget(TSharedPtr<FTransferStatistics> InStats, const TSharedRef<STableViewBase>& OwnerTable) const;
+		TSharedRef<ITableRow> OnGenerateActivityRowWidget(TSharedPtr<FOutboundTransferStatistics> InStats, const TSharedRef<STableViewBase>& OwnerTable) const;
 		void OnTransferStatisticsUpdated() const;
 	};
 }
diff --git a/Engine/Plugins/Messaging/UdpMessaging/Source/UdpMessaging/Private/Transport/UdpMessageProcessor.cpp b/Engine/Plugins/Messaging/UdpMessaging/Source/UdpMessaging/Private/Transport/UdpMessageProcessor.cpp
--- a/Engine/Plugins/Messaging/UdpMessaging/Source/UdpMessaging/Private/Transport/UdpMessageProcessor.cpp
+++ b/Engine/Plugins/Messaging/UdpMessaging/Source/UdpMessaging/Private/Transport/UdpMessageProcessor.cpp
@@ -77,9 +77,15 @@ bool ShouldErrorOnConnection(const FIPv4Endpoint& InEndpoint)
 	return false;
 }
 
-FOnTransferDataUpdated& OnSegmenterUpdated()
+FOnOutboundTransferDataUpdated& OnSegmenterUpdated()
 {
-	static FOnTransferDataUpdated OnTransferUpdated;
+	static FOnOutboundTransferDataUpdated OnTransferUpdated;
+	return OnTransferUpdated;
+}
+
+FOnInboundTransferDataUpdated& OnReassemblerUpdated()
+{
+	static FOnInboundTransferDataUpdated OnTransferUpdated;
 	return OnTransferUpdated;
 }
 
@@ -252,7 +258,7 @@ FMessageTransportStatistics FUdpMessageProcessor::GetStats(FGuid Node) const
 
 void FUdpMessageProcessor::SendSegmenterStatsToListeners(int32 MessageId, FGuid NodeId, const TSharedPtr<FUdpMessageSegmenter>& Segmenter)
 {
-	FOnTransferDataUpdated& SegmenterUpdatedDelegate = UE::Private::MessageProcessor::OnSegmenterUpdated();
+	FOnOutboundTransferDataUpdated& SegmenterUpdatedDelegate = UE::Private::MessageProcessor::OnSegmenterUpdated();
 	if(!SegmenterUpdatedDelegate.IsBound())
 	{
 		return;
@@ -717,6 +723,16 @@ void FUdpMessageProcessor::ProcessDataSegment(FInboundSegment& Segment, FNodeInf
 	NodeInfo.Statistics.TotalBytesReceived += DataChunk.Data.Num();
 	NodeInfo.Statistics.PacketsReceived++;
 	ReassembledMessage->Reassemble(DataChunk.SegmentNumber, DataChunk.SegmentOffset, DataChunk.Data, CurrentTime);
+	FOnInboundTransferDataUpdated& ReassemblerUpdated = UE::Private::MessageProcessor::OnReassemblerUpdated();
+	if(ReassemblerUpdated.IsBound())
+	{
+		ReassemblerUpdated.Broadcast(
+			{NodeInfo.NodeId,
+			 DataChunk.MessageId,
+			 UDP_MESSAGING_SEGMENT_SIZE * ReassembledMessage->GetPendingSegmentsCount(),
+			 ReassembledMessage->GetReceivedBytes()}
+			);
+	}
 
 	// Deliver or re-sequence message
 	if (!ReassembledMessage->IsComplete() || ReassembledMessage->IsDelivered())
diff --git a/Engine/Plugins/Messaging/UdpMessaging/Source/UdpMessaging/Private/Transport/UdpMessageProcessor.h b/Engine/Plugins/Messaging/UdpMessaging/Source/UdpMessaging/Private/Transport/UdpMessageProcessor.h
--- a/Engine/Plugins/Messaging/UdpMessaging/Source/UdpMessaging/Private/Transport/UdpMessageProcessor.h
+++ b/Engine/Plugins/Messaging/UdpMessaging/Source/UdpMessaging/Private/Transport/UdpMessageProcessor.h
@@ -111,9 +111,14 @@ struct FSentSegmentInfo
 namespace UE::Private::MessageProcessor
 {
 /**
- * Global delegate for handling segmenter updates.
+ * Global delegate for handling segmenter (aka sent data) updates.
  */
-FOnTransferDataUpdated& OnSegmenterUpdated();
+FOnOutboundTransferDataUpdated &OnSegmenterUpdated();
+
+/*
+ * Global delegate for handling reassembler (aka received data) updates.
+ */
+FOnInboundTransferDataUpdated& OnReassemblerUpdated();
 }
 
 /**
diff --git a/Engine/Plugins/Messaging/UdpMessaging/Source/UdpMessaging/Private/Transport/UdpReassembledMessage.h b/Engine/Plugins/Messaging/UdpMessaging/Source/UdpMessaging/Private/Transport/UdpReassembledMessage.h
--- a/Engine/Plugins/Messaging/UdpMessaging/Source/UdpMessaging/Private/Transport/UdpReassembledMessage.h
+++ b/Engine/Plugins/Messaging/UdpMessaging/Source/UdpMessaging/Private/Transport/UdpReassembledMessage.h
@@ -303,6 +303,11 @@ public:
 		return FString::Printf(TEXT("Sender=%s, Sequence=%llu"), *Sender.ToString(), Sequence);
 	}
 
+	int32 GetReceivedBytes() const
+	{
+		return ReceivedBytes;
+	}
+
 private:
 
 	/** Holds the message protocol version. */
diff --git a/Engine/Plugins/Messaging/UdpMessaging/Source/UdpMessaging/Private/UdpMessagingModule.cpp b/Engine/Plugins/Messaging/UdpMessaging/Source/UdpMessaging/Private/UdpMessagingModule.cpp
--- a/Engine/Plugins/Messaging/UdpMessaging/Source/UdpMessaging/Private/UdpMessagingModule.cpp
+++ b/Engine/Plugins/Messaging/UdpMessaging/Source/UdpMessaging/Private/UdpMessagingModule.cpp
@@ -397,11 +397,16 @@ public:
 		return MessageBridge->LookupAddress(MessageAddress);
 	}
 
-	virtual FOnTransferDataUpdated& OnTransferUpdatedFromThread() override
+	virtual FOnOutboundTransferDataUpdated& OnOutboundTransferUpdatedFromThread() override
 	{
 		return UE::Private::MessageProcessor::OnSegmenterUpdated();
 	}
 
+	virtual FOnInboundTransferDataUpdated& OnInboundTransferUpdatedFromThread() override
+	{
+		return UE::Private::MessageProcessor::OnReassemblerUpdated();
+	}
+
 	virtual TArray<FString> GetListeningAddresses() const override
 	{
 		if (auto Transport = WeakBridgeTransport.Pin())
diff --git a/Engine/Source/Runtime/Messaging/Public/INetworkMessagingExtension.h b/Engine/Source/Runtime/Messaging/Public/INetworkMessagingExtension.h
--- a/Engine/Source/Runtime/Messaging/Public/INetworkMessagingExtension.h
+++ b/Engine/Source/Runtime/Messaging/Public/INetworkMessagingExtension.h
@@ -51,7 +51,7 @@ struct FMessageTransportStatistics
 };
 
 /** Per-node per-message transfer statistics. */
-struct FTransferStatistics
+struct FOutboundTransferStatistics
 {
 	/** Unique Id for the target. */
 	FGuid  DestinationId;
@@ -69,7 +69,25 @@ struct FTransferStatistics
 	uint64 BytesAcknowledged;
 };
 
-DECLARE_MULTICAST_DELEGATE_OneParam(FOnTransferDataUpdated, FTransferStatistics);
+/** Per-node per-message transfer statistics. */
+struct FInboundTransferStatistics
+{
+	/** Unique Id for the target. */
+	FGuid  OriginId;
+
+	/** Message Id that uniquely identifies the inbound message. If MessageId == -1 then this message is not reliably sent or is not segmented over multiple packets.*/
+	int32  MessageId;
+	
+	/** Total bytes to send for the given MessageId */
+	uint64 BytesToReceive;
+
+	/** Total bytes sent thus far. */
+	int32 BytesReceived;
+};
+
+
+DECLARE_MULTICAST_DELEGATE_OneParam(FOnOutboundTransferDataUpdated, FOutboundTransferStatistics);
+DECLARE_MULTICAST_DELEGATE_OneParam(FOnInboundTransferDataUpdated, FInboundTransferStatistics);
 
 /**
  * Interface for the messaging module network extension
@@ -122,10 +140,16 @@ public:
 	virtual FGuid GetNodeIdFromAddress(const FMessageAddress& MessageAddress) const = 0;
 
 	/**
-	 * Delegate invoked when any transmission statistics are updated. This delegate may get called from another thread.
+	 * Delegate invoked when any transmission statistics are updated for outbound nodes (sent). This delegate may get called from another thread.
+	 * Please consider thread safety when receiving this delegate.
+	 */
+	virtual FOnOutboundTransferDataUpdated& OnOutboundTransferUpdatedFromThread() = 0;
+
+	/**
+	 * Delegate invoked when any transmission statistics are update for inbound (received). This delegate may get called from another thread.
 	 * Please consider thread safety when receiving this delegate.
 	 */
-	virtual FOnTransferDataUpdated& OnTransferUpdatedFromThread() = 0;
+	virtual FOnInboundTransferDataUpdated& OnInboundTransferUpdatedFromThread() = 0;
 
 	/**
 	 * Shutdown this messaging extension services for MessageBus
