Hook up sent data to graphs

diff --git a/Engine/Plugins/Developer/Concert/ConcertApp/MultiUserServer/Source/MultiUserServer/Private/Widgets/Clients/Browser/Models/ClientTransferStatisticsModel.cpp b/Engine/Plugins/Developer/Concert/ConcertApp/MultiUserServer/Source/MultiUserServer/Private/Widgets/Clients/Browser/Models/ClientTransferStatisticsModel.cpp
--- a/Engine/Plugins/Developer/Concert/ConcertApp/MultiUserServer/Source/MultiUserServer/Private/Widgets/Clients/Browser/Models/ClientTransferStatisticsModel.cpp
+++ b/Engine/Plugins/Developer/Concert/ConcertApp/MultiUserServer/Source/MultiUserServer/Private/Widgets/Clients/Browser/Models/ClientTransferStatisticsModel.cpp
@@ -41,17 +41,13 @@ namespace UE::MultiUserServer
 
 	FClientTransferStatisticsModel::FClientTransferStatisticsModel(const FMessageAddress& ClientAddress)
 		: ClientAddress(ClientAddress)
+		, OutboundStatTracker([this](const FOutboundTransferStatistics& Stats) { return IsSentToClient(Stats); }, [](const FOutboundTransferStatistics& Stats){ return Stats.BytesSent; })
+		, InboundStatTracker([this](const FInboundTransferStatistics& Stats) { return IsReceivedFromClient(Stats); }, [](const FInboundTransferStatistics& Stats){ return Stats.BytesReceived; })
 	{
-		for (int32 i = 0; i < static_cast<int32>(EConcertTransferStatistic::Count); ++i)
-		{
-			const EConcertTransferStatistic StatType = static_cast<EConcertTransferStatistic>(i);
-			TransferStatisticsTimelines.Add(StatType);
-			OnTimelineUpdatedDelegates.Add(StatType);
-		}
-		
 		if (INetworkMessagingExtension* Statistics = Private::ClientTransferStatisticsModel::GetMessagingStatistics())
 		{
-			Statistics->OnOutboundTransferUpdatedFromThread().AddRaw(this, &FClientTransferStatisticsModel::OnTransferUpdatedFromThread);
+			Statistics->OnOutboundTransferUpdatedFromThread().AddRaw(&OutboundStatTracker, &TClientTransferStatTracker<FOutboundTransferStatistics>::OnTransferUpdatedFromThread);
+			Statistics->OnInboundTransferUpdatedFromThread().AddRaw(&InboundStatTracker, &TClientTransferStatTracker<FInboundTransferStatistics>::OnTransferUpdatedFromThread);
 			TickHandle = FTSTicker::GetCoreTicker().AddTicker(FTickerDelegate::CreateRaw(this, &FClientTransferStatisticsModel::Tick));
 		}
 	}
@@ -60,28 +56,42 @@ namespace UE::MultiUserServer
 	{
 		if (INetworkMessagingExtension* Statistics = Private::ClientTransferStatisticsModel::GetMessagingStatistics())
 		{
-			Statistics->OnOutboundTransferUpdatedFromThread().RemoveAll(this);
+			Statistics->OnOutboundTransferUpdatedFromThread().RemoveAll(&OutboundStatTracker);
+			Statistics->OnInboundTransferUpdatedFromThread().RemoveAll(&InboundStatTracker);
 		}
 		
 		FTSTicker::GetCoreTicker().RemoveTicker(TickHandle);
 	}
 
-	void FClientTransferStatisticsModel::OnTransferUpdatedFromThread(FOutboundTransferStatistics TransferStatistics)
+	const TArray<FConcertTransferSamplePoint>& FClientTransferStatisticsModel::GetTransferStatTimeline(EConcertTransferStatistic StatisticType) const
 	{
-		if (AreRelevantStats(TransferStatistics))
-		{
-			AsyncStatQueue.Enqueue(FTransferStatItem{ FDateTime::Now(), TransferStatistics });
+		switch (StatisticType)
+		{
+		case EConcertTransferStatistic::SentToClient: return OutboundStatTracker.GetTransferStatisticsTimeline();
+		case EConcertTransferStatistic::ReceivedFromClient: return InboundStatTracker.GetTransferStatisticsTimeline();;
+		case EConcertTransferStatistic::Count:
+		default:
+			checkNoEntry();
+			return OutboundStatTracker.GetTransferStatisticsTimeline();
 		}
 	}
 
-	bool FClientTransferStatisticsModel::AreRelevantStats(const FOutboundTransferStatistics& TransferStatistics) const
+	FOnTransferTimelineUpdated& FClientTransferStatisticsModel::OnTransferTimelineUpdated(EConcertTransferStatistic StatisticType)
 	{
-		return IsSentToClient(TransferStatistics) || IsReceivedFromClient(TransferStatistics);
+		switch (StatisticType)
+		{
+		case EConcertTransferStatistic::SentToClient: return OutboundStatTracker.GetOnTimelineUpdatedDelegates();
+		case EConcertTransferStatistic::ReceivedFromClient: return InboundStatTracker.GetOnTimelineUpdatedDelegates();;
+		case EConcertTransferStatistic::Count:
+		default:
+			checkNoEntry();
+			return OutboundStatTracker.GetOnTimelineUpdatedDelegates();
+		}
 	}
 
 	bool FClientTransferStatisticsModel::IsSentToClient(const FOutboundTransferStatistics& Item) const
 	{
-		if (INetworkMessagingExtension* Statistics = Private::ClientTransferStatisticsModel::GetMessagingStatistics())
+		if (const INetworkMessagingExtension* Statistics = Private::ClientTransferStatisticsModel::GetMessagingStatistics())
 		{
 			const FGuid ClientNodeId = Statistics->GetNodeIdFromAddress(ClientAddress);
 			return ClientNodeId == Item.DestinationId;
@@ -89,134 +99,20 @@ namespace UE::MultiUserServer
 		return false;
 	}
 
-	bool FClientTransferStatisticsModel::IsReceivedFromClient(const FOutboundTransferStatistics& Item) const
+	bool FClientTransferStatisticsModel::IsReceivedFromClient(const FInboundTransferStatistics& Item) const
 	{
-		// TODO: UDP module is missing a way to find out
+		if (const INetworkMessagingExtension* Statistics = Private::ClientTransferStatisticsModel::GetMessagingStatistics())
+		{
+			const FGuid ClientNodeId = Statistics->GetNodeIdFromAddress(ClientAddress);
+			return ClientNodeId == Item.OriginId;
+		}
 		return false; 
 	}
 	
 	bool FClientTransferStatisticsModel::Tick(float DeltaTime)
 	{
-		constexpr int32 NumStats = static_cast<int32>(EConcertTransferStatistic::Count);
-		FConcertStatTypeFlags ConcertStatFlags(false, NumStats);
-
-		RemoveOldStatTimelines(ConcertStatFlags);
-		const bool bRemovedAnyGroups = RemoveOldGroupedStats();
-		
-		const FDateTime Now = FDateTime::Now();
-		FConcertTransferSamplePoint SamplingThisTick[NumStats] = { FConcertTransferSamplePoint{ Now }, FConcertTransferSamplePoint{ Now } };
-		static_assert(NumStats == 2);
-		while (const TOptional<FTransferStatItem> TransferStatistics = AsyncStatQueue.Dequeue())
-		{
-			constexpr int32 StatAsInt = static_cast<int32>(EConcertTransferStatistic::SentToClient);
-			ConcertStatFlags[StatAsInt] = true;
-			
-			UpdateStatTimeline(TransferStatistics.GetValue(), SamplingThisTick[StatAsInt]);
-			UpdateGroupedStats(TransferStatistics.GetValue());
-		}
-
-		if (bRemovedAnyGroups || ConcertStatFlags.CountSetBits() != 0)
-		{
-			for (int32 i = 0; i < NumStats; ++i)
-			{
-				const bool bHadStatThisUpdate = ConcertStatFlags[i]; 
-				if (bHadStatThisUpdate)
-				{
-					const EConcertTransferStatistic Stat = static_cast<EConcertTransferStatistic>(i);
-					TransferStatisticsTimelines[Stat].Add(SamplingThisTick[i]);
-					OnTimelineUpdatedDelegates[Stat].Broadcast();
-				}
-			}
-			OnGroupsUpdatedDelegate.Broadcast();
-		}
-
+		OutboundStatTracker.Tick();
+		InboundStatTracker.Tick();
 		return true;
 	}
-
-	namespace Private::ClientTransferStatisticsModel
-	{
-		bool IsTooOld(const FDateTime& StatCreationTime, const FDateTime& Now)
-		{
-			const FTimespan RetainTime = FTimespan::FromSeconds(60.f);
-			return RetainTime < Now - StatCreationTime;
-		}
-	}
-	
-	void FClientTransferStatisticsModel::RemoveOldStatTimelines(FConcertStatTypeFlags& ModifiedTimelines)
-	{
-		const FDateTime Now = FDateTime::Now();
-		for (auto TimelineIt = TransferStatisticsTimelines.CreateIterator(); TimelineIt; ++TimelineIt)
-		{
-			TArray<FConcertTransferSamplePoint>& SamplePoints = TimelineIt->Value;
-			bool bModifiedAny = false;
-			for (auto SampleIt = SamplePoints.CreateIterator(); SampleIt; ++SampleIt)
-			{
-				if (Private::ClientTransferStatisticsModel::IsTooOld(SampleIt->LocalTime, Now))
-				{
-					SampleIt.RemoveCurrent();
-					bModifiedAny = true;
-				}
-			}
-
-			const int32 Index = static_cast<int32>(TimelineIt->Key);
-			ModifiedTimelines[Index] = bModifiedAny;
-		}
-
-		for (auto IncompleteIt = IncompleteStatsUntilNow.CreateIterator(); IncompleteIt; ++IncompleteIt)
-		{
-			if (Private::ClientTransferStatisticsModel::IsTooOld(IncompleteIt->Value.LastUpdate, Now))
-			{
-				IncompleteIt.RemoveCurrent();
-			}
-		}
-	}
-
-	bool FClientTransferStatisticsModel::RemoveOldGroupedStats()
-	{
-		bool bRemovedAny = false;
-		const FDateTime Now = FDateTime::Now();
-		for (auto It = TransferStatisticsGroupedById.CreateIterator(); It; ++It)
-		{
-			const TSharedPtr<FOutboundTransferStatistics>& Stat = *It;
-			if (Private::ClientTransferStatisticsModel::IsTooOld(LastUpdateGroupUpdates[Stat->MessageId], Now))
-			{
-				LastUpdateGroupUpdates.Remove(Stat->MessageId);
-				It.RemoveCurrent();
-				bRemovedAny = true;
-			}
-		}
-		return bRemovedAny;
-	}
-
-	void FClientTransferStatisticsModel::UpdateStatTimeline(const FTransferStatItem& TransferStatistics, FConcertTransferSamplePoint& SamplingThisTick)
-	{
-		uint64 BytesSentSinceLastTime = TransferStatistics.BytesSent;
-		// Same message ID may have been queued multiple times
-		if (FIncompleteMessageData* ItemThisTick = IncompleteStatsUntilNow.Find(TransferStatistics.MessageId))
-		{
-			BytesSentSinceLastTime -= ItemThisTick->BytesTransferredSoFar;
-		}
-
-		IncompleteStatsUntilNow.Add(TransferStatistics.MessageId, { TransferStatistics.BytesSent, FDateTime::Now() });
-		SamplingThisTick.BytesTransferred += BytesSentSinceLastTime;
-	}
-
-	void FClientTransferStatisticsModel::UpdateGroupedStats(const FTransferStatItem& TransferStatistics)
-	{
-		const TSharedPtr<FOutboundTransferStatistics> NewValue = MakeShared<FOutboundTransferStatistics>(TransferStatistics);
-		const auto Pos = Algo::LowerBound(TransferStatisticsGroupedById, NewValue, [](const TSharedPtr<FOutboundTransferStatistics>& Value, const TSharedPtr<FOutboundTransferStatistics>& Check)
-		{
-			return Value->MessageId < Check->MessageId;
-		});
-		if (TransferStatisticsGroupedById.Num() > 0 && Pos < TransferStatisticsGroupedById.Num() && TransferStatisticsGroupedById[Pos] && TransferStatisticsGroupedById[Pos]->MessageId == NewValue->MessageId)
-		{
-			*TransferStatisticsGroupedById[Pos] = *NewValue;
-		}
-		else
-		{
-			TransferStatisticsGroupedById.Insert(NewValue, Pos);
-		}
-
-		LastUpdateGroupUpdates.Add(NewValue->MessageId, FDateTime::Now());
-	}
-}
+}
\ No newline at end of file
diff --git a/Engine/Plugins/Developer/Concert/ConcertApp/MultiUserServer/Source/MultiUserServer/Private/Widgets/Clients/Browser/Models/ClientTransferStatisticsModel.h b/Engine/Plugins/Developer/Concert/ConcertApp/MultiUserServer/Source/MultiUserServer/Private/Widgets/Clients/Browser/Models/ClientTransferStatisticsModel.h
--- a/Engine/Plugins/Developer/Concert/ConcertApp/MultiUserServer/Source/MultiUserServer/Private/Widgets/Clients/Browser/Models/ClientTransferStatisticsModel.h
+++ b/Engine/Plugins/Developer/Concert/ConcertApp/MultiUserServer/Source/MultiUserServer/Private/Widgets/Clients/Browser/Models/ClientTransferStatisticsModel.h
@@ -10,14 +10,168 @@
 
 namespace UE::MultiUserServer
 {
-	struct FTransferStatItem : FOutboundTransferStatistics
+	template<typename TStatBase>
+	class TClientTransferStatTracker
 	{
-		FDateTime LocalTime;
+		struct FIncompleteMessageData
+		{
+			uint64 BytesTransferredSoFar;
+			FDateTime LastUpdate;
+		};
+		
+		struct TStatItem : TStatBase
+		{
+			FDateTime LocalTime;
 
-		FTransferStatItem(const FDateTime& LocalTime, const FOutboundTransferStatistics& Data)
-			: FOutboundTransferStatistics(Data)
-			, LocalTime(LocalTime)
+			TStatItem(const FDateTime& LocalTime, TStatBase Data)
+				: TStatBase(Data)
+				, LocalTime(LocalTime)
+			{}
+		};
+
+		using FMessageId = decltype(FOutboundTransferStatistics::MessageId);
+		
+	public:
+
+		using FShouldInclude = TFunction<bool(const TStatBase&)>;
+		using FStatGetter = TFunction<uint64(const TStatBase&)>;
+
+		TClientTransferStatTracker(FShouldInclude ShouldIncludeFunc, FStatGetter StatGetterFunc)
+			: ShouldIncludeFunc(MoveTemp(ShouldIncludeFunc))
+			, StatGetterFunc(MoveTemp(StatGetterFunc))
 		{}
+
+		void OnTransferUpdatedFromThread(TStatBase TransferStatistics)
+		{
+			if (ShouldIncludeFunc(TransferStatistics))
+			{
+				AsyncStatQueue.Enqueue(TStatItem{ FDateTime::Now(), TransferStatistics });
+			}
+		}
+
+		void Tick()
+		{
+			const bool bRemovedOldStats = RemoveOldStatTimelines();
+			const bool bRemovedAnyGroups = RemoveOldGroupedStats();
+			const bool bUpdatedAnyTimeline = !AsyncStatQueue.IsEmpty();
+			
+			const FDateTime Now = FDateTime::Now();
+			FConcertTransferSamplePoint SamplingThisTick{ Now };
+			while (const TOptional<TStatItem> TransferStatistics = AsyncStatQueue.Dequeue())
+			{
+				UpdateStatTimeline(TransferStatistics.GetValue(), SamplingThisTick);
+				UpdateGroupedStats(TransferStatistics.GetValue());
+			}
+
+			if (bRemovedOldStats || bUpdatedAnyTimeline || bRemovedAnyGroups)
+			{
+				TransferStatisticsTimeline.Add(SamplingThisTick);
+				OnTimelineUpdatedDelegates.Broadcast();
+				
+				OnGroupsUpdatedDelegate.Broadcast();
+			}
+		}
+		
+		FOnTransferTimelineUpdated& GetOnTimelineUpdatedDelegates() { return OnTimelineUpdatedDelegates; }
+		FOnTransferGroupsUpdated& GetOnGroupsUpdatedDelegate() { return OnGroupsUpdatedDelegate; }
+		const TArray<FConcertTransferSamplePoint>& GetTransferStatisticsTimeline() const { return TransferStatisticsTimeline; }
+		const TArray<TSharedPtr<TStatBase>>& GetTransferStatisticsGroupedById() const { return TransferStatisticsGroupedById; }
+
+	private:
+		
+		TSpscQueue<TStatItem> AsyncStatQueue;
+		
+		TArray<FConcertTransferSamplePoint> TransferStatisticsTimeline;
+		TMap<FMessageId, FIncompleteMessageData> IncompleteStatsUntilNow;
+		
+		TArray<TSharedPtr<TStatBase>> TransferStatisticsGroupedById;
+		TMap<FMessageId, FDateTime> LastUpdateGroupUpdates;
+		
+		FOnTransferTimelineUpdated OnTimelineUpdatedDelegates;
+		FOnTransferGroupsUpdated OnGroupsUpdatedDelegate;
+
+		FShouldInclude ShouldIncludeFunc;
+		FStatGetter StatGetterFunc;
+		
+		static bool IsTooOld(const FDateTime& StatCreationTime, const FDateTime& Now)
+		{
+			const FTimespan RetainTime = FTimespan::FromSeconds(60.f);
+			return RetainTime < Now - StatCreationTime;
+		}
+		
+		bool RemoveOldStatTimelines()
+		{
+			bool bModifiedTimelines = false;
+			
+			const FDateTime Now = FDateTime::Now();
+			for (auto SampleIt = TransferStatisticsTimeline.CreateIterator(); SampleIt; ++SampleIt)
+			{
+				if (IsTooOld(SampleIt->LocalTime, Now))
+				{
+					SampleIt.RemoveCurrent();
+					bModifiedTimelines = true;
+				}
+			}
+
+			for (auto IncompleteIt = IncompleteStatsUntilNow.CreateIterator(); IncompleteIt; ++IncompleteIt)
+			{
+				if (IsTooOld(IncompleteIt->Value.LastUpdate, Now))
+				{
+					IncompleteIt.RemoveCurrent();
+				}
+			}
+
+			return bModifiedTimelines;
+		}
+		
+		bool RemoveOldGroupedStats()
+		{
+			bool bRemovedAny = false;
+			const FDateTime Now = FDateTime::Now();
+			for (auto It = TransferStatisticsGroupedById.CreateIterator(); It; ++It)
+			{
+				const TSharedPtr<TStatBase>& Stat = *It;
+				if (IsTooOld(LastUpdateGroupUpdates[Stat->MessageId], Now))
+				{
+					LastUpdateGroupUpdates.Remove(Stat->MessageId);
+					It.RemoveCurrent();
+					bRemovedAny = true;
+				}
+			}
+			return bRemovedAny;
+		}
+		
+		void UpdateStatTimeline(const TStatBase& TransferStatistics, FConcertTransferSamplePoint& SamplingThisTick)
+		{
+			uint64 BytesSentSinceLastTime = StatGetterFunc(TransferStatistics);
+			// Same message ID may have been queued multiple times
+			if (FIncompleteMessageData* ItemThisTick = IncompleteStatsUntilNow.Find(TransferStatistics.MessageId))
+			{
+				BytesSentSinceLastTime -= ItemThisTick->BytesTransferredSoFar;
+			}
+
+			IncompleteStatsUntilNow.Add(TransferStatistics.MessageId, { StatGetterFunc(TransferStatistics), FDateTime::Now() });
+			SamplingThisTick.BytesTransferred += BytesSentSinceLastTime;
+		}
+		
+		void UpdateGroupedStats(const TStatBase& TransferStatistics)
+		{
+			const TSharedPtr<TStatBase> NewValue = MakeShared<TStatBase>(TransferStatistics);
+			const auto Pos = Algo::LowerBound(TransferStatisticsGroupedById, NewValue, [](const TSharedPtr<TStatBase>& Value, const TSharedPtr<TStatBase>& Check)
+			{
+				return Value->MessageId < Check->MessageId;
+			});
+			if (TransferStatisticsGroupedById.Num() > 0 && Pos < TransferStatisticsGroupedById.Num() && TransferStatisticsGroupedById[Pos] && TransferStatisticsGroupedById[Pos]->MessageId == NewValue->MessageId)
+			{
+				*TransferStatisticsGroupedById[Pos] = *NewValue;
+			}
+			else
+			{
+				TransferStatisticsGroupedById.Insert(NewValue, Pos);
+			}
+
+			LastUpdateGroupUpdates.Add(NewValue->MessageId, FDateTime::Now());
+		}
 	};
 	
 	class FClientTransferStatisticsModel : public IClientTransferStatisticsModel
@@ -27,49 +181,25 @@ namespace UE::MultiUserServer
 		FClientTransferStatisticsModel(const FMessageAddress& ClientAddress);
 		virtual ~FClientTransferStatisticsModel() override;
 		
-		virtual const TArray<FConcertTransferSamplePoint>& GetTransferStatTimeline(EConcertTransferStatistic StatisticType) const override { return TransferStatisticsTimelines[StatisticType]; }
-		virtual const TArray<TSharedPtr<FOutboundTransferStatistics>>& GetTransferStatsGroupedById() const override { return TransferStatisticsGroupedById; }
-		virtual FOnTransferTimelineUpdated& OnTransferTimelineUpdated(EConcertTransferStatistic StatisticType) override { return OnTimelineUpdatedDelegates[StatisticType]; }
-		virtual FOnTransferGroupsUpdated& OnTransferGroupsUpdated() override { return OnGroupsUpdatedDelegate; }
+		virtual const TArray<FConcertTransferSamplePoint>& GetTransferStatTimeline(EConcertTransferStatistic StatisticType) const override;
+		virtual const TArray<TSharedPtr<FOutboundTransferStatistics>>& GetOutboundTransferStatsGroupedById() const override { return OutboundStatTracker.GetTransferStatisticsGroupedById(); }
+		virtual const TArray<TSharedPtr<FInboundTransferStatistics>>& GetInboundTransferStatsGroupedById() const override { return InboundStatTracker.GetTransferStatisticsGroupedById(); }
+		virtual FOnTransferTimelineUpdated& OnTransferTimelineUpdated(EConcertTransferStatistic StatisticType) override;
+		virtual FOnTransferGroupsUpdated& OnOutboundTransferGroupsUpdated() override { return OutboundStatTracker.GetOnGroupsUpdatedDelegate(); }
+		virtual FOnTransferGroupsUpdated& OnInboundTransferGroupsUpdated() override { return InboundStatTracker.GetOnGroupsUpdatedDelegate(); }
 
 	private:
 
-		using FConcertStatTypeFlags = TBitArray<TInlineAllocator<static_cast<int32>(EConcertTransferStatistic::Count)>>;
-		using FMessageId = decltype(FOutboundTransferStatistics::MessageId);
-		struct FIncompleteMessageData
-		{
-			uint64 BytesTransferredSoFar;
-			FDateTime LastUpdate;
-		};
-
 		const FMessageAddress ClientAddress;
-		TSpscQueue<FTransferStatItem> AsyncStatQueue;
-
-		// Timelines
-		TMap<EConcertTransferStatistic, TArray<FConcertTransferSamplePoint>> TransferStatisticsTimelines;
-		TMap<FMessageId, FIncompleteMessageData> IncompleteStatsUntilNow;
-
-		// Grouped
-		TArray<TSharedPtr<FOutboundTransferStatistics>> TransferStatisticsGroupedById;
-		TMap<FMessageId, FDateTime> LastUpdateGroupUpdates;
-		
-		// Delegate info
-		TMap<EConcertTransferStatistic, FOnTransferTimelineUpdated> OnTimelineUpdatedDelegates;
-		FOnTransferGroupsUpdated OnGroupsUpdatedDelegate;
 		FTSTicker::FDelegateHandle TickHandle;
 		
-		void OnTransferUpdatedFromThread(FOutboundTransferStatistics TransferStatistics);
-		bool AreRelevantStats(const FOutboundTransferStatistics& TransferStatistics) const;
 		bool IsSentToClient(const FOutboundTransferStatistics& Item) const;
-		bool IsReceivedFromClient(const FOutboundTransferStatistics& Item) const;
+		bool IsReceivedFromClient(const FInboundTransferStatistics& Item) const;
 		
 		bool Tick(float DeltaTime);
-		
-		void RemoveOldStatTimelines(FConcertStatTypeFlags& ModifiedTimelines);
-		bool RemoveOldGroupedStats();
-		
-		void UpdateStatTimeline(const FTransferStatItem& TransferStatistics, FConcertTransferSamplePoint& SamplingThisTick);
-		void UpdateGroupedStats(const FTransferStatItem& TransferStatistics);
+
+		TClientTransferStatTracker<FOutboundTransferStatistics> OutboundStatTracker;
+		TClientTransferStatTracker<FInboundTransferStatistics> InboundStatTracker;
 	};
 }
 
diff --git a/Engine/Plugins/Developer/Concert/ConcertApp/MultiUserServer/Source/MultiUserServer/Private/Widgets/Clients/Browser/Models/IClientTransferStatisticsModel.h b/Engine/Plugins/Developer/Concert/ConcertApp/MultiUserServer/Source/MultiUserServer/Private/Widgets/Clients/Browser/Models/IClientTransferStatisticsModel.h
--- a/Engine/Plugins/Developer/Concert/ConcertApp/MultiUserServer/Source/MultiUserServer/Private/Widgets/Clients/Browser/Models/IClientTransferStatisticsModel.h
+++ b/Engine/Plugins/Developer/Concert/ConcertApp/MultiUserServer/Source/MultiUserServer/Private/Widgets/Clients/Browser/Models/IClientTransferStatisticsModel.h
@@ -43,12 +43,16 @@ namespace UE::MultiUserServer
 		/** @return Contains the transfer statistics over time. Used for graphs. */
 		virtual const TArray<FConcertTransferSamplePoint>& GetTransferStatTimeline(EConcertTransferStatistic StatisticType) const = 0;
 		/** @return Sorted descending by message ID. Every message ID is unique. Used for list views. */
-		virtual const TArray<TSharedPtr<FOutboundTransferStatistics>>& GetTransferStatsGroupedById() const = 0;
+		virtual const TArray<TSharedPtr<FOutboundTransferStatistics>>& GetOutboundTransferStatsGroupedById() const = 0;
+		/** @return Sorted descending by message ID. Every message ID is unique. Used for list views. */
+		virtual const TArray<TSharedPtr<FInboundTransferStatistics>>& GetInboundTransferStatsGroupedById() const = 0;
 
-		/** Called when GetTransferStatisticsGroupedById changes */
+		/** Called when GetTransferStatTimeline changes */
 		virtual FOnTransferTimelineUpdated& OnTransferTimelineUpdated(EConcertTransferStatistic StatisticType) = 0;
-		/** Called when GetTransferStatisticsGroupedById changes */
-		virtual FOnTransferGroupsUpdated& OnTransferGroupsUpdated() = 0;
+		/** Called when GetTransferStatsGroupedById changes */
+		virtual FOnTransferGroupsUpdated& OnOutboundTransferGroupsUpdated() = 0;
+		/** Called when GetTransferStatsGroupedById changes */
+		virtual FOnTransferGroupsUpdated& OnInboundTransferGroupsUpdated() = 0;
 
 		virtual ~IClientTransferStatisticsModel() = default;
 	};
diff --git a/Engine/Plugins/Developer/Concert/ConcertApp/MultiUserServer/Source/MultiUserServer/Private/Widgets/Clients/Browser/SClientTransferStatTable.cpp b/Engine/Plugins/Developer/Concert/ConcertApp/MultiUserServer/Source/MultiUserServer/Private/Widgets/Clients/Browser/SClientTransferStatTable.cpp
--- a/Engine/Plugins/Developer/Concert/ConcertApp/MultiUserServer/Source/MultiUserServer/Private/Widgets/Clients/Browser/SClientTransferStatTable.cpp
+++ b/Engine/Plugins/Developer/Concert/ConcertApp/MultiUserServer/Source/MultiUserServer/Private/Widgets/Clients/Browser/SClientTransferStatTable.cpp
@@ -101,12 +101,12 @@ namespace UE::MultiUserServer
 	{
 		StatsModel = InStatsModel;
 
-		StatsModel->OnTransferGroupsUpdated().AddSP(this, &SClientTransferStatTable::OnTransferStatisticsUpdated);
+		StatsModel->OnOutboundTransferGroupsUpdated().AddSP(this, &SClientTransferStatTable::OnTransferStatisticsUpdated);
 
 		ChildSlot
 		[
 			SAssignNew(SegmenterListView, SListView<TSharedPtr<FOutboundTransferStatistics>>)
-			.ListItemsSource(&InStatsModel->GetTransferStatsGroupedById())
+			.ListItemsSource(&InStatsModel->GetOutboundTransferStatsGroupedById())
 			.OnGenerateRow(this, &SClientTransferStatTable::OnGenerateActivityRowWidget)
 			.SelectionMode(ESelectionMode::Multi)
 			.HeaderRow
