Add inbound segment table

diff --git a/Engine/Plugins/Developer/Concert/ConcertApp/MultiUserServer/Source/MultiUserServer/Private/Widgets/Clients/Browser/SClientTransferStatTable.cpp b/Engine/Plugins/Developer/Concert/ConcertApp/MultiUserServer/Source/MultiUserServer/Private/Widgets/Clients/Browser/SClientTransferStatTable.cpp
deleted file mode 100644
diff --git a/Engine/Plugins/Developer/Concert/ConcertApp/MultiUserServer/Source/MultiUserServer/Private/Widgets/Clients/Browser/SClientTransferStatTable.h b/Engine/Plugins/Developer/Concert/ConcertApp/MultiUserServer/Source/MultiUserServer/Private/Widgets/Clients/Browser/SClientTransferStatTable.h
deleted file mode 100644
diff --git a/Engine/Plugins/Developer/Concert/ConcertApp/MultiUserServer/Source/MultiUserServer/Private/Widgets/Clients/Browser/SConcertClientBrowser.cpp b/Engine/Plugins/Developer/Concert/ConcertApp/MultiUserServer/Source/MultiUserServer/Private/Widgets/Clients/Browser/SConcertClientBrowser.cpp
--- a/Engine/Plugins/Developer/Concert/ConcertApp/MultiUserServer/Source/MultiUserServer/Private/Widgets/Clients/Browser/SConcertClientBrowser.cpp
+++ b/Engine/Plugins/Developer/Concert/ConcertApp/MultiUserServer/Source/MultiUserServer/Private/Widgets/Clients/Browser/SConcertClientBrowser.cpp
@@ -380,9 +380,13 @@ TSharedPtr<SWidget> UE::MultiUserServer::SConcertClientBrowser::OnGetContextMenu
 		LOCTEXT("DisplayMode.NetworkGraph.Title", "Network graph"),
 		LOCTEXT("DisplayMode.NetworkGraph.Tooltip", "Show the up and down stream network traffic on a graph")
 		);
-	AddDisplayModeEntry(MenuBuilder, EClientDisplayMode::SegementTable,
-		LOCTEXT("DisplayMode.SegementTable.Title", "Segment table"),
-		LOCTEXT("DisplayMode.SegementTable.Tooltip", "A table displaying the messaging protocol's segments MessageId, Sent, Acked and Size in realtime.")
+	AddDisplayModeEntry(MenuBuilder, EClientDisplayMode::OutboundSegementTable,
+		LOCTEXT("DisplayMode.OutboundSegementTable.Title", "Outbound segment table"),
+		LOCTEXT("DisplayMode.SegementTable.Tooltip", "A table displaying the messaging protocol's outbound segments' MessageId, Sent, Acked and Size data in realtime.")
+		);
+	AddDisplayModeEntry(MenuBuilder, EClientDisplayMode::InboundSegmentTable,
+		LOCTEXT("DisplayMode.InboundSegementTable.Title", "Inbound Segment table"),
+		LOCTEXT("DisplayMode.SegementTable.Tooltip", "A table displaying the messaging protocol's inbound segments' MessageId, Received and Size data in realtime.")
 		);
 	return MenuBuilder.MakeWidget();
 }
diff --git a/Engine/Plugins/Developer/Concert/ConcertApp/MultiUserServer/Source/MultiUserServer/Private/Widgets/Clients/Browser/SConcertClientBrowserItem.cpp b/Engine/Plugins/Developer/Concert/ConcertApp/MultiUserServer/Source/MultiUserServer/Private/Widgets/Clients/Browser/SConcertClientBrowserItem.cpp
--- a/Engine/Plugins/Developer/Concert/ConcertApp/MultiUserServer/Source/MultiUserServer/Private/Widgets/Clients/Browser/SConcertClientBrowserItem.cpp
+++ b/Engine/Plugins/Developer/Concert/ConcertApp/MultiUserServer/Source/MultiUserServer/Private/Widgets/Clients/Browser/SConcertClientBrowserItem.cpp
@@ -6,10 +6,12 @@
 #include "ConcertServerStyle.h"
 #include "INetworkMessagingExtension.h"
 #include "SClientNetworkStats.h"
-#include "SClientTransferStatTable.h"
 #include "Graph/SClientNetworkGraphs.h"
 #include "Models/ClientTransferStatisticsModel.h"
 #include "Models/IClientNetworkStatisticsModel.h"
+#include "Table/SClientOutboundTransferStatTable.h"
+#include "Table/SClientInboundTransferStatTable.h"
+
 #include "Styling/StyleColors.h"
 #include "Widgets/Images/SImage.h"
 #include "Widgets/LayerManager/STooltipPresenter.h"
@@ -18,7 +20,6 @@
 #include "Widgets/SOverlay.h"
 #include "Widgets/Text/STextBlock.h"
 #include "Widgets/SBoxPanel.h"
-#include "Widgets/SNullWidget.h"
 
 #define LOCTEXT_NAMESPACE "UnrealMultiUserUI.SConcertClientBrowserItem"
 
@@ -26,7 +27,7 @@ void UE::MultiUserServer::SConcertClientBrowserItem::Construct(const FArguments&
 {
 	Item = MoveTemp(InClientItem);
 	StatModel = MoveTemp(InStatModel);
-	TransferStatsModel = MakeShared<FClientTransferStatisticsModel>(InClientItem->ClientAddress);
+	TransferStatsModel = MakeShared<FClientTransferStatisticsModel>(Item->ClientAddress);
 	HighlightText = InArgs._HighlightText;
 
 	ChildSlot
@@ -143,8 +144,13 @@ TSharedRef<SWidget> UE::MultiUserServer::SConcertClientBrowserItem::CreateConten
 		]
 		+SOverlay::Slot().Padding(4.f, 0.f)
 		[
-			SNew(SClientTransferStatTable, TransferStatsModel.ToSharedRef())
-			.Visibility_Lambda([this](){ return GetDisplayMode() == EClientDisplayMode::SegementTable ? EVisibility::Visible : EVisibility::Collapsed; })
+			SNew(SClientOutboundTransferStatTable, TransferStatsModel.ToSharedRef())
+			.Visibility_Lambda([this](){ return GetDisplayMode() == EClientDisplayMode::OutboundSegementTable ? EVisibility::Visible : EVisibility::Collapsed; })
+		]
+		+SOverlay::Slot().Padding(4.f, 0.f)
+		[
+			SNew(SClientInboundTransferStatTable, TransferStatsModel.ToSharedRef())
+			.Visibility_Lambda([this](){ return GetDisplayMode() == EClientDisplayMode::InboundSegmentTable ? EVisibility::Visible : EVisibility::Collapsed; })
 		];
 }
 
diff --git a/Engine/Plugins/Developer/Concert/ConcertApp/MultiUserServer/Source/MultiUserServer/Private/Widgets/Clients/Browser/SConcertClientBrowserItem.h b/Engine/Plugins/Developer/Concert/ConcertApp/MultiUserServer/Source/MultiUserServer/Private/Widgets/Clients/Browser/SConcertClientBrowserItem.h
--- a/Engine/Plugins/Developer/Concert/ConcertApp/MultiUserServer/Source/MultiUserServer/Private/Widgets/Clients/Browser/SConcertClientBrowserItem.h
+++ b/Engine/Plugins/Developer/Concert/ConcertApp/MultiUserServer/Source/MultiUserServer/Private/Widgets/Clients/Browser/SConcertClientBrowserItem.h
@@ -21,8 +21,10 @@ namespace UE::MultiUserServer
 	{
 		/** Displays the sent and read packets */
 		NetworkGraph = 0,
-		/** Displays a table showing MessageId, Sent, Acked, and Size updated in realtime. */
-		SegementTable
+		/** Displays a table showing MessageId, Sent, Acked, and Size updated in realtime for outbound UDP segments. */
+		OutboundSegementTable = 1,
+		/** Displays a table showing MessageId, Received, and Size updated in realtime for inbound UDP segments. */
+		InboundSegmentTable = 2
 	};
 	
 	class SConcertClientBrowserItem : public SCompoundWidget
diff --git a/Engine/Plugins/Developer/Concert/ConcertApp/MultiUserServer/Source/MultiUserServer/Private/Widgets/Clients/Browser/Table/SClientInboundTransferStatTable.cpp b/Engine/Plugins/Developer/Concert/ConcertApp/MultiUserServer/Source/MultiUserServer/Private/Widgets/Clients/Browser/Table/SClientInboundTransferStatTable.cpp
new file mode 100644
--- /dev/null
+++ b/Engine/Plugins/Developer/Concert/ConcertApp/MultiUserServer/Source/MultiUserServer/Private/Widgets/Clients/Browser/Table/SClientInboundTransferStatTable.cpp
@@ -0,0 +1,123 @@
+// Copyright Epic Games, Inc. All Rights Reserved.
+
+#include "SClientInboundTransferStatTable.h"
+
+#include "INetworkMessagingExtension.h"
+#include "Widgets/Clients/Browser/Models/IClientTransferStatisticsModel.h"
+#include "Widgets/SBoxPanel.h"
+#include "Widgets/SNullWidget.h"
+#include "Widgets/Text/STextBlock.h"
+#include "Widgets/Views/SHeaderRow.h"
+#include "Widgets/Views/SListView.h"
+
+#define LOCTEXT_NAMESPACE "UnrealMultiUserUI.SClientInboundTransferStatTable"
+
+namespace UE::MultiUserServer
+{
+	class SClientInboundTransferStatTableRow : public SMultiColumnTableRow<TSharedRef<FOutboundTransferStatistics>>
+	{
+		using Super = SMultiColumnTableRow<TSharedRef<FOutboundTransferStatistics>>;
+	public:
+		
+		inline static const FName HeaderIdName_MessageId			= TEXT("MessageId");
+		inline static const FName HeaderIdName_ReceivedSegments	= TEXT("Received");
+		inline static const FName HeaderIdName_TotalSize			= TEXT("Size");
+		
+		SLATE_BEGIN_ARGS(SClientInboundTransferStatTableRow) { }
+			SLATE_ARGUMENT(TSharedPtr<FInboundTransferStatistics>, Stats)
+		SLATE_END_ARGS()
+		
+		void Construct(const FArguments& InArgs, const TSharedRef<STableViewBase>& InOwerTableView)
+		{
+			Stats = InArgs._Stats;
+			Super::Construct(Super::FArguments(), InOwerTableView);
+		}
+		
+	private:
+
+		TSharedPtr<FInboundTransferStatistics> Stats;
+
+		virtual TSharedRef<SWidget> GenerateWidgetForColumn(const FName& ColumnName) override
+		{
+			using namespace Private;
+			if (HeaderIdName_MessageId == ColumnName)
+			{
+				return SNew(SHorizontalBox)
+					+ SHorizontalBox::Slot()
+					.HAlign(HAlign_Left)
+					.VAlign(VAlign_Center)
+					.AutoWidth()
+					[
+						SNew(STextBlock)
+						.Text(FText::AsNumber(Stats->MessageId))
+					];
+			}
+			if (HeaderIdName_ReceivedSegments == ColumnName)
+			{
+				return SNew(SHorizontalBox)
+					+ SHorizontalBox::Slot()
+					.HAlign(HAlign_Left)
+					.VAlign(VAlign_Center)
+					.AutoWidth()
+					[
+						SNew(STextBlock)
+						.Text_Lambda([this](){ return FText::AsNumber(Stats->BytesReceived); })
+					];
+			}
+			if (HeaderIdName_TotalSize == ColumnName)
+			{
+				return SNew(SHorizontalBox)
+					+ SHorizontalBox::Slot()
+					.HAlign(HAlign_Left)
+					.VAlign(VAlign_Center)
+					.AutoWidth()
+					[
+						SNew(STextBlock)
+						.Text(FText::AsNumber(Stats->BytesToReceive))
+					];
+			}
+			return SNullWidget::NullWidget;
+		}
+	};
+	
+	void SClientInboundTransferStatTable::Construct(const FArguments& InArgs, TSharedRef<IClientTransferStatisticsModel> InStatsModel)
+	{
+		StatsModel = InStatsModel;
+
+		StatsModel->OnOutboundTransferGroupsUpdated().AddSP(this, &SClientInboundTransferStatTable::OnTransferStatisticsUpdated);
+
+		ChildSlot
+		[
+			SAssignNew(SegmenterListView, SListView<TSharedPtr<FInboundTransferStatistics>>)
+			.ListItemsSource(&InStatsModel->GetInboundTransferStatsGroupedById())
+			.OnGenerateRow(this, &SClientInboundTransferStatTable::OnGenerateActivityRowWidget)
+			.SelectionMode(ESelectionMode::Multi)
+			.HeaderRow
+			(
+				SNew(SHeaderRow)
+				+ SHeaderRow::Column(SClientInboundTransferStatTableRow::HeaderIdName_MessageId)
+				.FillWidth(2.f)
+				.DefaultLabel(LOCTEXT("HeaderName_MessageId", "Id"))
+				+ SHeaderRow::Column(SClientInboundTransferStatTableRow::HeaderIdName_ReceivedSegments)
+				.FillWidth(2.f)
+				.DefaultLabel(LOCTEXT("HeaderName_AckSegments", "Received"))
+				+ SHeaderRow::Column(SClientInboundTransferStatTableRow::HeaderIdName_TotalSize)
+				.FillWidth(2.f)
+				.DefaultLabel(LOCTEXT("HeaderName_TotalSize", "Size"))
+			)
+		];
+	}
+
+	TSharedRef<ITableRow> SClientInboundTransferStatTable::OnGenerateActivityRowWidget(TSharedPtr<FInboundTransferStatistics> InStats, const TSharedRef<STableViewBase>& OwnerTable) const
+	{
+		return SNew(SClientInboundTransferStatTableRow, OwnerTable)
+			.Stats(InStats);;
+	}
+
+	void SClientInboundTransferStatTable::OnTransferStatisticsUpdated() const
+	{
+		SegmenterListView->RequestListRefresh();
+	}
+}
+
+#undef LOCTEXT_NAMESPACE
diff --git a/Engine/Plugins/Developer/Concert/ConcertApp/MultiUserServer/Source/MultiUserServer/Private/Widgets/Clients/Browser/Table/SClientInboundTransferStatTable.h b/Engine/Plugins/Developer/Concert/ConcertApp/MultiUserServer/Source/MultiUserServer/Private/Widgets/Clients/Browser/Table/SClientInboundTransferStatTable.h
new file mode 100644
--- /dev/null
+++ b/Engine/Plugins/Developer/Concert/ConcertApp/MultiUserServer/Source/MultiUserServer/Private/Widgets/Clients/Browser/Table/SClientInboundTransferStatTable.h
@@ -0,0 +1,43 @@
+// Copyright Epic Games, Inc. All Rights Reserved.
+
+#pragma once
+
+#include "CoreMinimal.h"
+#include "INetworkMessagingExtension.h"
+#include "Widgets/SCompoundWidget.h"
+
+class ITableRow;
+template<typename T> class SListView;
+class STableViewBase;
+struct FOutboundTransferStatistics;
+
+namespace UE::MultiUserServer
+{
+	class IClientTransferStatisticsModel;
+
+	/**
+	 * Displays FOutboundTransferStatistics as they are updated by INetworkMessagingExtension.
+	 */
+	class SClientInboundTransferStatTable : public SCompoundWidget
+	{
+	public:
+
+		SLATE_BEGIN_ARGS(SClientInboundTransferStatTable)
+		{}
+		SLATE_END_ARGS()
+
+		void Construct(const FArguments& InArgs, TSharedRef<IClientTransferStatisticsModel> InStatsModel);
+
+	private:
+
+		/** Tells us when the transfer stats have changed */
+		TSharedPtr<IClientTransferStatisticsModel> StatsModel;
+		/** Displays the transfer stats */
+		TSharedPtr<SListView<TSharedPtr<FInboundTransferStatistics>>> SegmenterListView;
+
+		TSharedRef<ITableRow> OnGenerateActivityRowWidget(TSharedPtr<FInboundTransferStatistics> InStats, const TSharedRef<STableViewBase>& OwnerTable) const;
+		void OnTransferStatisticsUpdated() const;
+	};
+}
+
+
diff --git a/Engine/Plugins/Developer/Concert/ConcertApp/MultiUserServer/Source/MultiUserServer/Private/Widgets/Clients/Browser/Table/SClientOutboundTransferStatTable.cpp b/Engine/Plugins/Developer/Concert/ConcertApp/MultiUserServer/Source/MultiUserServer/Private/Widgets/Clients/Browser/Table/SClientOutboundTransferStatTable.cpp
new file mode 100644
--- /dev/null
+++ b/Engine/Plugins/Developer/Concert/ConcertApp/MultiUserServer/Source/MultiUserServer/Private/Widgets/Clients/Browser/Table/SClientOutboundTransferStatTable.cpp
@@ -0,0 +1,143 @@
+// Copyright Epic Games, Inc. All Rights Reserved.
+
+#include "SClientOutboundTransferStatTable.h"
+
+#include "INetworkMessagingExtension.h"
+#include "Widgets/Clients/Browser/Models/IClientTransferStatisticsModel.h"
+#include "Widgets/SBoxPanel.h"
+#include "Widgets/SNullWidget.h"
+#include "Widgets/Text/STextBlock.h"
+#include "Widgets/Views/SHeaderRow.h"
+#include "Widgets/Views/SListView.h"
+
+#define LOCTEXT_NAMESPACE "UnrealMultiUserUI.SClientTransferStatTable"
+
+namespace UE::MultiUserServer
+{
+	namespace Private
+	{
+		const FName HeaderIdName_MessageId    = TEXT("MessageId");
+		const FName HeaderIdName_SentSegments = TEXT("Sent");
+		const FName HeaderIdName_AckSegments  = TEXT("SegmentsAck");
+		const FName HeaderIdName_TotalSize    = TEXT("Size");
+		const FName HeaderIdName_DataRate     = TEXT("DataRate");
+	}
+
+	class SClientOutboundTransferStatTableRow : public SMultiColumnTableRow<TSharedRef<FOutboundTransferStatistics>>
+	{
+		using Super = SMultiColumnTableRow<TSharedRef<FOutboundTransferStatistics>>;
+	public:
+		
+		SLATE_BEGIN_ARGS(SClientOutboundTransferStatTableRow) { }
+			SLATE_ARGUMENT(TSharedPtr<FOutboundTransferStatistics>, Stats)
+		SLATE_END_ARGS()
+		
+		void Construct(const FArguments& InArgs, const TSharedRef<STableViewBase>& InOwerTableView)
+		{
+			Stats = InArgs._Stats;
+			Super::Construct(Super::FArguments(), InOwerTableView);
+		}
+		
+	private:
+
+		TSharedPtr<FOutboundTransferStatistics> Stats;
+
+		virtual TSharedRef<SWidget> GenerateWidgetForColumn(const FName& ColumnName) override
+		{
+			using namespace Private;
+			if (HeaderIdName_MessageId == ColumnName)
+			{
+				return SNew(SHorizontalBox)
+					+ SHorizontalBox::Slot()
+					.HAlign(HAlign_Left)
+					.VAlign(VAlign_Center)
+					.AutoWidth()
+					[
+						SNew(STextBlock)
+						.Text(FText::AsNumber(Stats->MessageId))
+					];
+			}
+			if (HeaderIdName_SentSegments == ColumnName)
+			{
+				return SNew(SHorizontalBox)
+					+ SHorizontalBox::Slot()
+					.HAlign(HAlign_Left)
+					.VAlign(VAlign_Center)
+					.AutoWidth()
+					[
+						SNew(STextBlock)
+						.Text_Lambda([this](){ return FText::AsNumber(Stats->BytesSent); })
+					];
+			}
+			if (HeaderIdName_AckSegments == ColumnName)
+			{
+				return SNew(SHorizontalBox)
+					+ SHorizontalBox::Slot()
+					.HAlign(HAlign_Left)
+					.VAlign(VAlign_Center)
+					.AutoWidth()
+					[
+						SNew(STextBlock)
+						.Text_Lambda([this](){ return FText::AsNumber(Stats->BytesAcknowledged); })
+					];
+			}
+			if (HeaderIdName_TotalSize == ColumnName)
+			{
+				return SNew(SHorizontalBox)
+					+ SHorizontalBox::Slot()
+					.HAlign(HAlign_Left)
+					.VAlign(VAlign_Center)
+					.AutoWidth()
+					[
+						SNew(STextBlock)
+						.Text(FText::AsNumber(Stats->BytesToSend))
+					];
+			}
+			return SNullWidget::NullWidget;
+		}
+	};
+	
+	void SClientOutboundTransferStatTable::Construct(const FArguments& InArgs, TSharedRef<IClientTransferStatisticsModel> InStatsModel)
+	{
+		StatsModel = InStatsModel;
+
+		StatsModel->OnOutboundTransferGroupsUpdated().AddSP(this, &SClientOutboundTransferStatTable::OnTransferStatisticsUpdated);
+
+		ChildSlot
+		[
+			SAssignNew(SegmenterListView, SListView<TSharedPtr<FOutboundTransferStatistics>>)
+			.ListItemsSource(&InStatsModel->GetOutboundTransferStatsGroupedById())
+			.OnGenerateRow(this, &SClientOutboundTransferStatTable::OnGenerateActivityRowWidget)
+			.SelectionMode(ESelectionMode::Multi)
+			.HeaderRow
+			(
+				SNew(SHeaderRow)
+				+ SHeaderRow::Column(Private::HeaderIdName_MessageId)
+				.FillWidth(2.f)
+				.DefaultLabel(LOCTEXT("HeaderName_MessageId", "Id"))
+				+ SHeaderRow::Column(Private::HeaderIdName_SentSegments)
+				.FillWidth(2.f)
+				.DefaultLabel(LOCTEXT("HeaderName_SentSegments", "Sent"))
+				+ SHeaderRow::Column(Private::HeaderIdName_AckSegments)
+				.FillWidth(2.f)
+				.DefaultLabel(LOCTEXT("HeaderName_AckSegments", "Ack"))
+				+ SHeaderRow::Column(Private::HeaderIdName_TotalSize)
+				.FillWidth(2.f)
+				.DefaultLabel(LOCTEXT("HeaderName_TotalSize", "Size"))
+			)
+		];
+	}
+
+	TSharedRef<ITableRow> SClientOutboundTransferStatTable::OnGenerateActivityRowWidget(TSharedPtr<FOutboundTransferStatistics> InStats, const TSharedRef<STableViewBase>& OwnerTable) const
+	{
+		return SNew(SClientOutboundTransferStatTableRow, OwnerTable)
+			.Stats(InStats);;
+	}
+
+	void SClientOutboundTransferStatTable::OnTransferStatisticsUpdated() const
+	{
+		SegmenterListView->RequestListRefresh();
+	}
+}
+
+#undef LOCTEXT_NAMESPACE
diff --git a/Engine/Plugins/Developer/Concert/ConcertApp/MultiUserServer/Source/MultiUserServer/Private/Widgets/Clients/Browser/Table/SClientOutboundTransferStatTable.h b/Engine/Plugins/Developer/Concert/ConcertApp/MultiUserServer/Source/MultiUserServer/Private/Widgets/Clients/Browser/Table/SClientOutboundTransferStatTable.h
new file mode 100644
--- /dev/null
+++ b/Engine/Plugins/Developer/Concert/ConcertApp/MultiUserServer/Source/MultiUserServer/Private/Widgets/Clients/Browser/Table/SClientOutboundTransferStatTable.h
@@ -0,0 +1,42 @@
+// Copyright Epic Games, Inc. All Rights Reserved.
+
+#pragma once
+
+#include "CoreMinimal.h"
+#include "Widgets/SCompoundWidget.h"
+
+class ITableRow;
+template<typename T> class SListView;
+class STableViewBase;
+struct FOutboundTransferStatistics;
+
+namespace UE::MultiUserServer
+{
+	class IClientTransferStatisticsModel;
+
+	/**
+	 * Displays FOutboundTransferStatistics as they are updated by INetworkMessagingExtension.
+	 */
+	class SClientOutboundTransferStatTable : public SCompoundWidget
+	{
+	public:
+
+		SLATE_BEGIN_ARGS(SClientOutboundTransferStatTable)
+		{}
+		SLATE_END_ARGS()
+
+		void Construct(const FArguments& InArgs, TSharedRef<IClientTransferStatisticsModel> InStatsModel);
+
+	private:
+
+		/** Tells us when the transfer stats have changed */
+		TSharedPtr<IClientTransferStatisticsModel> StatsModel;
+		/** Displays the transfer stats */
+		TSharedPtr<SListView<TSharedPtr<FOutboundTransferStatistics>>> SegmenterListView;
+
+		TSharedRef<ITableRow> OnGenerateActivityRowWidget(TSharedPtr<FOutboundTransferStatistics> InStats, const TSharedRef<STableViewBase>& OwnerTable) const;
+		void OnTransferStatisticsUpdated() const;
+	};
+}
+
+
