Extract FClientBrowserItem functionality to IConcertBrowserItem interface

diff --git a/Engine/Plugins/Developer/Concert/ConcertApp/MultiUserServer/Source/MultiUserServer/Private/Widgets/Clients/Browser/ClientBrowserItem.cpp b/Engine/Plugins/Developer/Concert/ConcertApp/MultiUserServer/Source/MultiUserServer/Private/Widgets/Clients/Browser/ClientBrowserItem.cpp
new file mode 100644
--- /dev/null
+++ b/Engine/Plugins/Developer/Concert/ConcertApp/MultiUserServer/Source/MultiUserServer/Private/Widgets/Clients/Browser/ClientBrowserItem.cpp
@@ -0,0 +1,57 @@
+// Copyright Epic Games, Inc. All Rights Reserved.
+
+#include "ClientBrowserItem.h"
+
+#include "Models/ClientTransferStatisticsModel.h"
+#include "Widgets/Clients/Util/EndpointToUserNameCache.h"
+
+#define LOCTEXT_NAMESPACE "UnrealMultiUserUI.FClientBrowserItem"
+
+namespace UE::MultiUserServer
+{
+	FClientBrowserItem::FClientBrowserItem(TSharedRef<IClientNetworkStatisticsModel> NetworkStatisticsModel, TSharedRef<FEndpointToUserNameCache> UserLookup, const FMessageAddress& ClientAddress, const FGuid& MessageNodeId)
+		: FConcertBrowserItemCommonImpl(MoveTemp(NetworkStatisticsModel))
+		, UserLookup(MoveTemp(UserLookup))
+		, TransferStatisticsModel(MakeShared<FClientTransferStatisticsModel>(ClientAddress))
+		, ClientAddress(ClientAddress)
+		, MessageNodeId(MessageNodeId)
+	{}
+	
+	FString FClientBrowserItem::GetDisplayName() const
+	{
+		if (const TOptional<FConcertClientInfo> ClientInfo = UserLookup->GetClientInfoFromNodeId(MessageNodeId))
+		{
+			return ClientInfo->DisplayName;
+		}
+		
+		const FString NodeIdAsString = MessageNodeId.ToString(EGuidFormats::DigitsWithHyphens);
+		int32 Index = INDEX_NONE;
+		const bool bFound = NodeIdAsString.FindChar('-', Index);
+		// Avoid making the name too long by only showing the first few digits of the node ID
+		return FString::Printf(TEXT("Admin (%s)"), bFound ? *NodeIdAsString.Left(Index) : *NodeIdAsString); 
+	}
+
+	FText FClientBrowserItem::GetToolTip() const
+	{
+		if (const TOptional<FConcertClientInfo> ClientInfo = UserLookup->GetClientInfoFromNodeId(MessageNodeId))
+		{
+			return FText::Format(
+				LOCTEXT("Name.Available.TooltipFmt", "NodeID: {0}\nAddress ID: {1}"),
+				FText::FromString(MessageNodeId.ToString(EGuidFormats::DigitsWithHyphens)),
+				FText::FromString(ClientAddress.ToString())
+				);
+		}
+		return FText::Format(
+			LOCTEXT("Name.NotAvailable.TooltipFmt", "This client's display information becomes available after joining a session.\nNodeID: {0}\nAddress ID: {1}"),
+			FText::FromString(MessageNodeId.ToString(EGuidFormats::DigitsWithHyphens)),
+			FText::FromString(ClientAddress.ToString())
+			);
+	}
+
+	TSharedRef<IClientTransferStatisticsModel> FClientBrowserItem::GetTransferStatistics() const
+	{
+		return TransferStatisticsModel;
+	}
+}
+
+#undef LOCTEXT_NAMESPACE
\ No newline at end of file
diff --git a/Engine/Plugins/Developer/Concert/ConcertApp/MultiUserServer/Source/MultiUserServer/Private/Widgets/Clients/Browser/ClientBrowserItem.h b/Engine/Plugins/Developer/Concert/ConcertApp/MultiUserServer/Source/MultiUserServer/Private/Widgets/Clients/Browser/ClientBrowserItem.h
--- a/Engine/Plugins/Developer/Concert/ConcertApp/MultiUserServer/Source/MultiUserServer/Private/Widgets/Clients/Browser/ClientBrowserItem.h
+++ b/Engine/Plugins/Developer/Concert/ConcertApp/MultiUserServer/Source/MultiUserServer/Private/Widgets/Clients/Browser/ClientBrowserItem.h
@@ -4,33 +4,57 @@
 
 #include "CoreMinimal.h"
 #include "ConcertMessageData.h"
+#include "IConcertBrowserItem.h"
 #include "IMessageContext.h"
 
+class FEndpointToUserNameCache;
+
 namespace UE::MultiUserServer
 {
+	class IClientTransferStatisticsModel;
+	class FClientTransferStatisticsModel;
+
 	DECLARE_DELEGATE_RetVal(TOptional<FConcertClientInfo>, FGetClientInfo);
 	
-	struct FClientBrowserItem
+	class FClientBrowserItem : public FConcertBrowserItemCommonImpl
 	{
-		/** Getter for getting the latest client info. */
-		FGetClientInfo GetClientInfo;
+	public:
+
+		FClientBrowserItem(TSharedRef<IClientNetworkStatisticsModel> NetworkStatisticsModel, TSharedRef<FEndpointToUserNameCache> UserLookup, const FMessageAddress& ClientAddress, const FGuid& MessageNodeId);
 		
+		void OnJoinSession(const FGuid& SessionId) { CurrentSession = SessionId; }
+		void OnLeaveSession() { CurrentSession.Reset(); }
+		const TOptional<FGuid>& GetCurrentSession() const { return CurrentSession; }
+		const FGuid& GetMessageNodeId() const { return MessageNodeId; }
+
+		void OnClientDisconnected() { bIsDisconnected = true; }
+		void OnClientReconnected() { bIsDisconnected = false; }
+		/** Whether this client is no longer connected with the server */
+		bool IsDisconnected() const { return bIsDisconnected; }
+		
+		virtual FString GetDisplayName() const override;
+		virtual FText GetToolTip() const override;
+		virtual FMessageAddress GetMessageAddress() const override { return ClientAddress; }
+		virtual TSharedRef<IClientTransferStatisticsModel> GetTransferStatistics() const override;
+
+	private:
+
+		/** Used to gets the client's name once it becomes available */
+		TSharedRef<FEndpointToUserNameCache> UserLookup;
+
+		/** Keeps track of the clients transfer statistics */
+		TSharedRef<FClientTransferStatisticsModel> TransferStatisticsModel;
+
 		/** Address of this client as used by the UDP messaging sytem */
 		FMessageAddress ClientAddress;
+		
 		/** ID of this client in the UDP messaging system */
 		FGuid MessageNodeId;
-
+		
 		/** The session this client currently is in */
 		TOptional<FGuid> CurrentSession;
 
 		/** Whether this client is no longer connected with the server */
 		bool bIsDisconnected = false;
-
-		FClientBrowserItem(FGetClientInfo GetClientInfo, const FMessageAddress& ClientAddress, const FGuid& MessageNodeId, TOptional<FGuid> CurrentSession = {})
-			: GetClientInfo(MoveTemp(GetClientInfo))
-			, ClientAddress(ClientAddress)
-			, MessageNodeId(MessageNodeId)
-			, CurrentSession(CurrentSession)
-		{}
 	};
 }
diff --git a/Engine/Plugins/Developer/Concert/ConcertApp/MultiUserServer/Source/MultiUserServer/Private/Widgets/Clients/Browser/IConcertBrowserItem.cpp b/Engine/Plugins/Developer/Concert/ConcertApp/MultiUserServer/Source/MultiUserServer/Private/Widgets/Clients/Browser/IConcertBrowserItem.cpp
new file mode 100644
--- /dev/null
+++ b/Engine/Plugins/Developer/Concert/ConcertApp/MultiUserServer/Source/MultiUserServer/Private/Widgets/Clients/Browser/IConcertBrowserItem.cpp
@@ -0,0 +1,51 @@
+// Copyright Epic Games, Inc. All Rights Reserved.
+
+#include "IConcertBrowserItem.h"
+
+#include "Models/IClientNetworkStatisticsModel.h"
+
+#include "IMessageContext.h"
+#include "INetworkMessagingExtension.h"
+
+namespace UE::MultiUserServer
+{
+	FConcertBrowserItemCommonImpl::FConcertBrowserItemCommonImpl(TSharedRef<IClientNetworkStatisticsModel> NetworkStatisticsModel)
+			: NetworkStatisticsModel(MoveTemp(NetworkStatisticsModel))
+	{}
+
+	void FConcertBrowserItemCommonImpl::SetDisplayMode(EConcertBrowserItemDisplayMode Value)
+	{
+		ItemDisplayMode = Value;
+	}
+	
+	EConcertBrowserItemDisplayMode FConcertBrowserItemCommonImpl::GetDisplayMode() const
+	{
+		return ItemDisplayMode;
+	}
+
+	void FConcertBrowserItemCommonImpl::AppendSearchTerms(TArray<FString>& SearchTerms) const
+	{
+		SearchTerms.Add(GetDisplayName());
+		const TOptional<FMessageTransportStatistics> Stats = GetLatestNetworkStatistics();
+		SearchTerms.Add(NetworkStatistics::FormatIPv4AsString(Stats));
+		
+		if (Stats)
+		{
+			SearchTerms.Add(NetworkStatistics::FormatTotalBytesSent(*Stats));
+			SearchTerms.Add(NetworkStatistics::FormatTotalBytesReceived(*Stats));
+			SearchTerms.Add(NetworkStatistics::FormatAverageRTT(*Stats));
+			SearchTerms.Add(NetworkStatistics::FormatBytesInflight(*Stats));
+			SearchTerms.Add(NetworkStatistics::FormatTotalBytesLost(*Stats));
+		}
+	}
+
+	TOptional<FMessageTransportStatistics> FConcertBrowserItemCommonImpl::GetLatestNetworkStatistics() const
+	{
+		return NetworkStatisticsModel->GetLatestNetworkStatistics(GetMessageAddress());
+	}
+	
+	bool FConcertBrowserItemCommonImpl::IsOnline() const
+	{
+		return NetworkStatisticsModel->IsOnline(GetMessageAddress());
+	}
+}
diff --git a/Engine/Plugins/Developer/Concert/ConcertApp/MultiUserServer/Source/MultiUserServer/Private/Widgets/Clients/Browser/IConcertBrowserItem.h b/Engine/Plugins/Developer/Concert/ConcertApp/MultiUserServer/Source/MultiUserServer/Private/Widgets/Clients/Browser/IConcertBrowserItem.h
new file mode 100644
--- /dev/null
+++ b/Engine/Plugins/Developer/Concert/ConcertApp/MultiUserServer/Source/MultiUserServer/Private/Widgets/Clients/Browser/IConcertBrowserItem.h
@@ -0,0 +1,72 @@
+// Copyright Epic Games, Inc. All Rights Reserved.
+
+#pragma once
+
+#include "CoreMinimal.h"
+
+struct FMessageTransportStatistics;
+struct FMessageAddress;
+
+namespace UE::MultiUserServer
+{
+	class IClientNetworkStatisticsModel;
+	class IClientTransferStatisticsModel;
+
+	enum class EConcertBrowserItemDisplayMode : uint8
+	{
+		/** Displays the sent and read packets */
+		NetworkGraph = 0,
+		/** Displays a table showing MessageId, Sent, Acked, and Size updated in realtime for outbound UDP segments. */
+		OutboundSegementTable = 1,
+		/** Displays a table showing MessageId, Received, and Size updated in realtime for inbound UDP segments. */
+		InboundSegmentTable = 2
+	};
+
+	/** Implemented by items that can be displayed in the Concert Network Browser */
+	class IConcertBrowserItem
+	{
+	public:
+
+		/** The title of the thumbnail */
+		virtual FString GetDisplayName() const = 0;
+		/** What to display when the thumbnail is hovered */
+		virtual FText GetToolTip() const = 0;
+		/** The address of the corresponding item. */
+		virtual FMessageAddress GetMessageAddress() const = 0;
+		
+		virtual void SetDisplayMode(EConcertBrowserItemDisplayMode Value) = 0;
+		virtual EConcertBrowserItemDisplayMode GetDisplayMode() const = 0;
+
+		/** The network statistics for this item */
+		virtual TSharedRef<IClientTransferStatisticsModel> GetTransferStatistics() const = 0;
+
+		/** Adds terms this item can be searched by */
+		virtual void AppendSearchTerms(TArray<FString>& SearchTerms) const = 0;
+
+		/** Gets the latest network statistics for this item if they are available; most likely unavailable when IsOnline returns false. */
+		virtual TOptional<FMessageTransportStatistics> GetLatestNetworkStatistics() const = 0;
+		/** Whether this client is currently reachable */
+		virtual bool IsOnline() const = 0;
+
+		virtual ~IConcertBrowserItem() = default;
+	};
+
+	/** Util for implementing shared logic. */
+	class FConcertBrowserItemCommonImpl : public IConcertBrowserItem
+	{
+	public:
+
+		FConcertBrowserItemCommonImpl(TSharedRef<IClientNetworkStatisticsModel> NetworkStatisticsModel);
+
+		virtual void SetDisplayMode(EConcertBrowserItemDisplayMode Value) override;
+		virtual EConcertBrowserItemDisplayMode GetDisplayMode() const override;
+		virtual void AppendSearchTerms(TArray<FString>& SearchTerms) const override;
+		virtual TOptional<FMessageTransportStatistics> GetLatestNetworkStatistics() const override;
+		virtual bool IsOnline() const override;
+
+	private:
+
+		TSharedRef<IClientNetworkStatisticsModel> NetworkStatisticsModel;
+		EConcertBrowserItemDisplayMode ItemDisplayMode = EConcertBrowserItemDisplayMode::NetworkGraph;
+	};
+}
diff --git a/Engine/Plugins/Developer/Concert/ConcertApp/MultiUserServer/Source/MultiUserServer/Private/Widgets/Clients/Browser/Models/ClientBrowserModel.cpp b/Engine/Plugins/Developer/Concert/ConcertApp/MultiUserServer/Source/MultiUserServer/Private/Widgets/Clients/Browser/Models/ClientBrowserModel.cpp
--- a/Engine/Plugins/Developer/Concert/ConcertApp/MultiUserServer/Source/MultiUserServer/Private/Widgets/Clients/Browser/Models/ClientBrowserModel.cpp
+++ b/Engine/Plugins/Developer/Concert/ConcertApp/MultiUserServer/Source/MultiUserServer/Private/Widgets/Clients/Browser/Models/ClientBrowserModel.cpp
@@ -12,9 +12,10 @@
 #include "Features/IModularFeatures.h"
 #include "INetworkMessagingExtension.h"
 
-UE::MultiUserServer::FClientBrowserModel::FClientBrowserModel(TSharedRef<IConcertServer> InServer, TSharedRef<FEndpointToUserNameCache> ClientInfoCache)
+UE::MultiUserServer::FClientBrowserModel::FClientBrowserModel(TSharedRef<IConcertServer> InServer, TSharedRef<FEndpointToUserNameCache> ClientInfoCache, TSharedRef<IClientNetworkStatisticsModel> NetworkStatisticsModel)
 	: Server(MoveTemp(InServer))
 	, ClientInfoCache(MoveTemp(ClientInfoCache))
+	, NetworkStatisticsModel(MoveTemp(NetworkStatisticsModel))
 {
 	ConcertServerEvents::OnLiveSessionCreated().AddRaw(this, &FClientBrowserModel::OnLiveSessionCreated);
 	ConcertServerEvents::OnLiveSessionDestroyed().AddRaw(this, &FClientBrowserModel::OnLiveSessionDestroyed);
@@ -73,7 +74,7 @@ void UE::MultiUserServer::FClientBrowserModel::SetKeepClientsAfterDisconnect(boo
 	}
 	for (auto It = Clients.CreateIterator(); It; ++It)
 	{
-		if (It->Get()->bIsDisconnected)
+		if (It->Get()->IsDisconnected())
 		{
 			OnClientListChanged().Broadcast(*It, EClientUpdateType::Removed);
 			It.RemoveCurrent();
@@ -152,11 +153,18 @@ void UE::MultiUserServer::FClientBrowserModel::UpdateClientSessionId(IConcertSer
 	const FGuid NodeId = Extension.GetNodeIdFromAddress(Address);
 	const int32 Index = Clients.IndexOfByPredicate([&NodeId](const TSharedPtr<FClientBrowserItem>& Item)
 	{
-		return Item->MessageNodeId == NodeId;
+		return Item->GetMessageNodeId() == NodeId;
 	});
 	if (Clients.IsValidIndex(Index))
 	{
-		Clients[Index]->CurrentSession = SessionId;
+		if (SessionId)
+		{
+			Clients[Index]->OnJoinSession(*SessionId);
+		}
+		else
+		{
+			Clients[Index]->OnLeaveSession();
+		}
 	}
 }
 
@@ -196,19 +204,17 @@ void UE::MultiUserServer::FClientBrowserModel::AddClientAdminEndpoint(const FCon
 	const FGuid NodeId = MessagingExtension->GetNodeIdFromAddress(Address);
 	const TSharedPtr<FClientBrowserItem>* ExistingItem = Clients.FindByPredicate([&NodeId](const TSharedPtr<FClientBrowserItem>& Item)
 	{
-		return Item->MessageNodeId == NodeId;
+		return Item->GetMessageNodeId() == NodeId;
 	});
 	if (ExistingItem)
 	{
-		ExistingItem->Get()->bIsDisconnected = false;
+		ExistingItem->Get()->OnClientReconnected();
 		return;
 	}
 	
 	const TSharedRef<FClientBrowserItem> Item = MakeShared<FClientBrowserItem>(
-		FGetClientInfo::CreateLambda([this, NodeId]()
-		{
-			return ClientInfoCache->GetClientInfoFromNodeId(NodeId);
-		}),
+		NetworkStatisticsModel,
+		ClientInfoCache,
 		Address,
 		NodeId
 		);
@@ -227,7 +233,7 @@ void UE::MultiUserServer::FClientBrowserModel::RemoveClientAdminEndpoint(const F
 	
 	const int32 Index = Clients.IndexOfByPredicate([NodeId](const TSharedPtr<FClientBrowserItem>& Item)
 	{
-		return Item->MessageNodeId == *NodeId;
+		return Item->GetMessageNodeId() == *NodeId;
 	});
 	
 	if (ensure(Clients.IsValidIndex(Index)))
@@ -236,7 +242,7 @@ void UE::MultiUserServer::FClientBrowserModel::RemoveClientAdminEndpoint(const F
 		const TSharedPtr<FClientBrowserItem> Item = Clients[Index]; 
 		if (bKeepClientsAfterDisconnect)
 		{
-			Item->bIsDisconnected = true;
+			Item->OnClientDisconnected();
 			return;
 		}
 		
diff --git a/Engine/Plugins/Developer/Concert/ConcertApp/MultiUserServer/Source/MultiUserServer/Private/Widgets/Clients/Browser/Models/ClientBrowserModel.h b/Engine/Plugins/Developer/Concert/ConcertApp/MultiUserServer/Source/MultiUserServer/Private/Widgets/Clients/Browser/Models/ClientBrowserModel.h
--- a/Engine/Plugins/Developer/Concert/ConcertApp/MultiUserServer/Source/MultiUserServer/Private/Widgets/Clients/Browser/Models/ClientBrowserModel.h
+++ b/Engine/Plugins/Developer/Concert/ConcertApp/MultiUserServer/Source/MultiUserServer/Private/Widgets/Clients/Browser/Models/ClientBrowserModel.h
@@ -13,11 +13,13 @@ class IConcertServer;
 
 namespace UE::MultiUserServer
 {
+	class IClientNetworkStatisticsModel;
+
 	class FClientBrowserModel : public IClientBrowserModel
 	{
 	public:
 
-		FClientBrowserModel(TSharedRef<IConcertServer> InServer, TSharedRef<FEndpointToUserNameCache> ClientInfoCache);
+		FClientBrowserModel(TSharedRef<IConcertServer> InServer, TSharedRef<FEndpointToUserNameCache> ClientInfoCache, TSharedRef<IClientNetworkStatisticsModel> NetworkStatisticsModel);
 		virtual ~FClientBrowserModel() override;
 
 		//~ Begin IClientBrowserModel Interface
@@ -38,6 +40,7 @@ namespace UE::MultiUserServer
 
 		TSharedRef<IConcertServer> Server;
 		TSharedRef<FEndpointToUserNameCache> ClientInfoCache;
+		TSharedRef<IClientNetworkStatisticsModel> NetworkStatisticsModel;
 
 		TArray<TSharedPtr<FClientBrowserItem>> Clients;
 		/** Needed for when an admin endpoint disconnects - querying won't work anymore then. */
diff --git a/Engine/Plugins/Developer/Concert/ConcertApp/MultiUserServer/Source/MultiUserServer/Private/Widgets/Clients/Browser/Models/IClientBrowserModel.h b/Engine/Plugins/Developer/Concert/ConcertApp/MultiUserServer/Source/MultiUserServer/Private/Widgets/Clients/Browser/Models/IClientBrowserModel.h
--- a/Engine/Plugins/Developer/Concert/ConcertApp/MultiUserServer/Source/MultiUserServer/Private/Widgets/Clients/Browser/Models/IClientBrowserModel.h
+++ b/Engine/Plugins/Developer/Concert/ConcertApp/MultiUserServer/Source/MultiUserServer/Private/Widgets/Clients/Browser/Models/IClientBrowserModel.h
@@ -11,7 +11,7 @@ struct FMessageAddress;
 
 namespace UE::MultiUserServer
 {
-	struct FClientBrowserItem;
+	class FClientBrowserItem;
 	
 	/** Decouples the UI from the server functions. */
 	class IClientBrowserModel
diff --git a/Engine/Plugins/Developer/Concert/ConcertApp/MultiUserServer/Source/MultiUserServer/Private/Widgets/Clients/Browser/Models/IClientNetworkStatisticsModel.cpp b/Engine/Plugins/Developer/Concert/ConcertApp/MultiUserServer/Source/MultiUserServer/Private/Widgets/Clients/Browser/Models/IClientNetworkStatisticsModel.cpp
new file mode 100644
--- /dev/null
+++ b/Engine/Plugins/Developer/Concert/ConcertApp/MultiUserServer/Source/MultiUserServer/Private/Widgets/Clients/Browser/Models/IClientNetworkStatisticsModel.cpp
@@ -0,0 +1,52 @@
+// Copyright Epic Games, Inc. All Rights Reserved.
+
+#include "IClientNetworkStatisticsModel.h"
+
+#include "INetworkMessagingExtension.h"
+#include "Math/UnitConversion.h"
+
+namespace UE::MultiUserServer::NetworkStatistics
+{
+	FString FormatIPv4AsString(const TOptional<FMessageTransportStatistics>& Stats)
+	{
+		const FString Fallback(TEXT("No IP available"));
+		if (Stats)
+		{
+			const FString DisplayString = Stats->IPv4AsString.IsEmpty()
+				? Fallback
+				: Stats->IPv4AsString;
+			return DisplayString;
+		}
+		return Fallback;
+	}
+		
+	FString FormatTotalBytesSent(const FMessageTransportStatistics& Stats)
+	{
+		const FNumericUnit<uint64> Unit = FUnitConversion::QuantizeUnitsToBestFit(Stats.TotalBytesSent, EUnit::Bytes);
+		return FString::Printf(TEXT("%lld %s"), Unit.Value, FUnitConversion::GetUnitDisplayString(Unit.Units));
+	}
+	
+	FString FormatTotalBytesReceived(const FMessageTransportStatistics& Stats)
+	{
+		const FNumericUnit<uint64> Unit = FUnitConversion::QuantizeUnitsToBestFit(Stats.TotalBytesReceived, EUnit::Bytes);
+		return FString::Printf(TEXT("%lld %s"), Unit.Value, FUnitConversion::GetUnitDisplayString(Unit.Units));
+	}
+	
+	FString FormatAverageRTT(const FMessageTransportStatistics& Stats)
+	{
+		const FNumericUnit<uint64> Unit = FUnitConversion::QuantizeUnitsToBestFit(static_cast<uint64>(Stats.AverageRTT.GetTotalMilliseconds()), EUnit::Milliseconds);
+		return FString::Printf(TEXT("%lld %s"), Unit.Value, FUnitConversion::GetUnitDisplayString(Unit.Units));;
+	}
+	
+	FString FormatBytesInflight(const FMessageTransportStatistics& Stats)
+	{
+		const FNumericUnit<uint64> Unit = FUnitConversion::QuantizeUnitsToBestFit(Stats.BytesInflight, EUnit::Bytes);
+		return FString::Printf(TEXT("%lld %s"), Unit.Value, FUnitConversion::GetUnitDisplayString(Unit.Units));
+	}
+	
+	FString FormatTotalBytesLost(const FMessageTransportStatistics& Stats)
+	{
+		const FNumericUnit<uint64> Unit = FUnitConversion::QuantizeUnitsToBestFit(Stats.TotalBytesLost, EUnit::Bytes);
+		return FString::Printf(TEXT("%lld %s"), Unit.Value, FUnitConversion::GetUnitDisplayString(Unit.Units));
+	}
+}
diff --git a/Engine/Plugins/Developer/Concert/ConcertApp/MultiUserServer/Source/MultiUserServer/Private/Widgets/Clients/Browser/Models/IClientNetworkStatisticsModel.h b/Engine/Plugins/Developer/Concert/ConcertApp/MultiUserServer/Source/MultiUserServer/Private/Widgets/Clients/Browser/Models/IClientNetworkStatisticsModel.h
--- a/Engine/Plugins/Developer/Concert/ConcertApp/MultiUserServer/Source/MultiUserServer/Private/Widgets/Clients/Browser/Models/IClientNetworkStatisticsModel.h
+++ b/Engine/Plugins/Developer/Concert/ConcertApp/MultiUserServer/Source/MultiUserServer/Private/Widgets/Clients/Browser/Models/IClientNetworkStatisticsModel.h
@@ -23,6 +23,18 @@ namespace UE::MultiUserServer
 
 		virtual ~IClientNetworkStatisticsModel() = default;
 	};
+
+	// Should move IClientNetworkStatisticsModel into the this namespace for more cohesion
+	namespace NetworkStatistics
+	{
+		FString FormatIPv4AsString(const TOptional<FMessageTransportStatistics>& Stats);
+		
+		FString FormatTotalBytesSent(const FMessageTransportStatistics& Stats);
+		FString FormatTotalBytesReceived(const FMessageTransportStatistics& Stats);
+		FString FormatAverageRTT(const FMessageTransportStatistics& Stats);
+		FString FormatBytesInflight(const FMessageTransportStatistics& Stats);
+		FString FormatTotalBytesLost(const FMessageTransportStatistics& Stats);
+	}
 }
 
 
diff --git a/Engine/Plugins/Developer/Concert/ConcertApp/MultiUserServer/Source/MultiUserServer/Private/Widgets/Clients/Browser/SClientNetworkStats.cpp b/Engine/Plugins/Developer/Concert/ConcertApp/MultiUserServer/Source/MultiUserServer/Private/Widgets/Clients/Browser/SClientNetworkStats.cpp
--- a/Engine/Plugins/Developer/Concert/ConcertApp/MultiUserServer/Source/MultiUserServer/Private/Widgets/Clients/Browser/SClientNetworkStats.cpp
+++ b/Engine/Plugins/Developer/Concert/ConcertApp/MultiUserServer/Source/MultiUserServer/Private/Widgets/Clients/Browser/SClientNetworkStats.cpp
@@ -2,12 +2,10 @@
 
 #include "SClientNetworkStats.h"
 
-#include "ConcertServerStyle.h"
 #include "INetworkMessagingExtension.h"
 #include "Math/UnitConversion.h"
 #include "Widgets/Clients/Browser/Models/ClientNetworkStatisticsModel.h"
 
-#include "Widgets/Layout/SBorder.h"
 #include "Widgets/SBoxPanel.h"
 #include "Widgets/Text/STextBlock.h"
 
@@ -15,10 +13,9 @@
 
 namespace UE::MultiUserServer
 {
-	void SClientNetworkStats::Construct(const FArguments& InArgs, const FMessageAddress& InNodeAddress, TSharedRef<IClientNetworkStatisticsModel> InStatisticModel)
+	void SClientNetworkStats::Construct(const FArguments& InArgs)
 	{
-		NodeAddress = InNodeAddress;
-		StatisticModel = MoveTemp(InStatisticModel);
+		NetworkStatistics = InArgs._NetworkStatistics;
 		
 		HighlightText = InArgs._HighlightText;
 		check(HighlightText);
@@ -34,7 +31,7 @@ namespace UE::MultiUserServer
 	{
 		SCompoundWidget::Tick(AllottedGeometry, InCurrentTime, InDeltaTime);
 
-		if (const TOptional<FMessageTransportStatistics> Stats = StatisticModel->GetLatestNetworkStatistics(NodeAddress))
+		if (const TOptional<FMessageTransportStatistics> Stats = NetworkStatistics.Get())
 		{
 			UpdateStatistics(*Stats);
 		}
@@ -49,47 +46,13 @@ namespace UE::MultiUserServer
 		}
 	}
 
-	void SClientNetworkStats::AppendSearchTerms(TArray<FString>& SearchTerms) const
-	{
-		SearchTerms.Add(SendText->GetText().ToString());
-		SearchTerms.Add(ReceiveText->GetText().ToString());
-		SearchTerms.Add(RoundTripTimeText->GetText().ToString());
-		SearchTerms.Add(InflightText->GetText().ToString());
-		SearchTerms.Add(LossText->GetText().ToString());
-	}
-
 	void SClientNetworkStats::UpdateStatistics(const FMessageTransportStatistics& Statistics)
 	{
-		// Send
-		{
-			const FNumericUnit<uint64> Unit = FUnitConversion::QuantizeUnitsToBestFit(Statistics.TotalBytesSent, EUnit::Bytes);
-			const FString DisplayString = FString::Printf(TEXT("%lld %s"), Unit.Value, FUnitConversion::GetUnitDisplayString(Unit.Units));
-			SendText->SetText(FText::FromString(DisplayString));
-		}
-		// Receive
-		{
-			const FNumericUnit<uint64> Unit = FUnitConversion::QuantizeUnitsToBestFit(Statistics.TotalBytesReceived, EUnit::Bytes);
-			const FString DisplayString = FString::Printf(TEXT("%lld %s"), Unit.Value, FUnitConversion::GetUnitDisplayString(Unit.Units));
-			ReceiveText->SetText(FText::FromString(DisplayString));
-		}
-		// RTT
-		{
-			const FNumericUnit<uint64> Unit = FUnitConversion::QuantizeUnitsToBestFit(static_cast<uint64>(Statistics.AverageRTT.GetTotalMilliseconds()), EUnit::Milliseconds);
-			const FString DisplayString = FString::Printf(TEXT("%lld %s"), Unit.Value, FUnitConversion::GetUnitDisplayString(Unit.Units));
-			RoundTripTimeText->SetText(FText::FromString(DisplayString));
-		}
-		// Inflight
-		{
-			const FNumericUnit<uint64> Unit = FUnitConversion::QuantizeUnitsToBestFit(Statistics.BytesInflight, EUnit::Bytes);
-			const FString DisplayString = FString::Printf(TEXT("%lld %s"), Unit.Value, FUnitConversion::GetUnitDisplayString(Unit.Units));
-			InflightText->SetText(FText::FromString(DisplayString));
-		}
-		// Loss
-		{
-			const FNumericUnit<uint64> Unit = FUnitConversion::QuantizeUnitsToBestFit(Statistics.TotalBytesLost, EUnit::Bytes);
-			const FString DisplayString = FString::Printf(TEXT("%lld %s"), Unit.Value, FUnitConversion::GetUnitDisplayString(Unit.Units));
-			LossText->SetText(FText::FromString(DisplayString));
-		}
+		SendText->SetText(FText::FromString(NetworkStatistics::FormatTotalBytesSent(Statistics)));
+		ReceiveText->SetText(FText::FromString(NetworkStatistics::FormatTotalBytesReceived(Statistics)));
+		RoundTripTimeText->SetText(FText::FromString(NetworkStatistics::FormatAverageRTT(Statistics)));
+		InflightText->SetText(FText::FromString(NetworkStatistics::FormatBytesInflight(Statistics)));
+		LossText->SetText(FText::FromString(NetworkStatistics::FormatTotalBytesLost(Statistics)));
 	}
 
 	TSharedRef<SWidget> SClientNetworkStats::CreateStatTable()
@@ -97,9 +60,10 @@ namespace UE::MultiUserServer
 		const TSharedRef<SHorizontalBox> Content = SNew(SHorizontalBox)
 			.ToolTipText_Lambda([this]()
 			{
-				const TOptional<FMessageTransportStatistics> Stats = StatisticModel->GetLatestNetworkStatistics(NodeAddress);
+				const TOptional<FMessageTransportStatistics> Stats = NetworkStatistics.Get();
 				return Stats.IsSet() ? FText::GetEmpty() : LOCTEXT("ErrorGettingStats", "Error getting stats");
 			});
+		
 		AddStatistic(Content, LOCTEXT("SentLabel", "Sent"), LOCTEXT("SentTooltip", "Total bytes sent to this client"), SendText);
 		AddStatistic(Content, LOCTEXT("ReceiveLabel", "Received"), LOCTEXT("ReceiveTooltip", "Total bytes received from this client"), ReceiveText);
 		AddStatistic(Content, LOCTEXT("RttLabel", "RTT"), LOCTEXT("RttTooltip", "Round trip time"), RoundTripTimeText);
diff --git a/Engine/Plugins/Developer/Concert/ConcertApp/MultiUserServer/Source/MultiUserServer/Private/Widgets/Clients/Browser/SClientNetworkStats.h b/Engine/Plugins/Developer/Concert/ConcertApp/MultiUserServer/Source/MultiUserServer/Private/Widgets/Clients/Browser/SClientNetworkStats.h
--- a/Engine/Plugins/Developer/Concert/ConcertApp/MultiUserServer/Source/MultiUserServer/Private/Widgets/Clients/Browser/SClientNetworkStats.h
+++ b/Engine/Plugins/Developer/Concert/ConcertApp/MultiUserServer/Source/MultiUserServer/Private/Widgets/Clients/Browser/SClientNetworkStats.h
@@ -4,43 +4,36 @@
 
 #include "CoreMinimal.h"
 #include "IMessageContext.h"
+#include "INetworkMessagingExtension.h"
 #include "Widgets/DeclarativeSyntaxSupport.h"
 #include "Widgets/SBoxPanel.h"
 #include "Widgets/SCompoundWidget.h"
 
 class STextBlock;
-struct FMessageTransportStatistics;
 
 namespace UE::MultiUserServer
 {
-	class IClientNetworkStatisticsModel;
-	
 	/**
 	 * Displays statistics about a client connection in a table like format: send, receive, RTT, Inflight, and Loss.
 	 */
 	class SClientNetworkStats : public SCompoundWidget
 	{
 	public:
-
+		
 		SLATE_BEGIN_ARGS(SClientNetworkStats){}
 			SLATE_ARGUMENT(TSharedPtr<FText>, HighlightText)
+			SLATE_ATTRIBUTE(TOptional<FMessageTransportStatistics>, NetworkStatistics)
 		SLATE_END_ARGS()
 
-		void Construct(const FArguments& InArgs, const FMessageAddress& InNodeAddress, TSharedRef<IClientNetworkStatisticsModel> InStatisticModel);
+		void Construct(const FArguments& InArgs);
 
 		//~ Begin SCompoundWidget Interface		
 		virtual void Tick(const FGeometry& AllottedGeometry, const double InCurrentTime, const float InDeltaTime) override;
 		//~ End SCompoundWidget Interface
-		
-		/** Appends the statistics to the search terms */
-		void AppendSearchTerms(TArray<FString>& SearchTerms) const;
 
 	private:
 
-		/** The ID getting visualised */
-		FMessageAddress NodeAddress;
-		/** Used to obtain settings */
-		TSharedPtr<IClientNetworkStatisticsModel> StatisticModel;
+		TAttribute<TOptional<FMessageTransportStatistics>> NetworkStatistics;
 		
 		/** The text to highlight */
 		TSharedPtr<FText> HighlightText;
diff --git a/Engine/Plugins/Developer/Concert/ConcertApp/MultiUserServer/Source/MultiUserServer/Private/Widgets/Clients/Browser/SConcertBrowserItem.cpp b/Engine/Plugins/Developer/Concert/ConcertApp/MultiUserServer/Source/MultiUserServer/Private/Widgets/Clients/Browser/SConcertBrowserItem.cpp
new file mode 100644
--- /dev/null
+++ b/Engine/Plugins/Developer/Concert/ConcertApp/MultiUserServer/Source/MultiUserServer/Private/Widgets/Clients/Browser/SConcertBrowserItem.cpp
@@ -0,0 +1,183 @@
+// Copyright Epic Games, Inc. All Rights Reserved.
+
+#include "SConcertBrowserItem.h"
+
+#include "ClientBrowserItem.h"
+#include "ConcertServerStyle.h"
+#include "IConcertBrowserItem.h"
+#include "INetworkMessagingExtension.h"
+#include "SClientNetworkStats.h"
+#include "Graph/SClientNetworkGraphs.h"
+#include "Models/ClientTransferStatisticsModel.h"
+#include "Models/IClientNetworkStatisticsModel.h"
+#include "Table/SClientOutboundTransferStatTable.h"
+#include "Table/SClientInboundTransferStatTable.h"
+
+#include "Styling/StyleColors.h"
+#include "Widgets/Images/SImage.h"
+#include "Widgets/LayerManager/STooltipPresenter.h"
+#include "Widgets/Layout/SBorder.h"
+#include "Widgets/Layout/SScaleBox.h"
+#include "Widgets/SOverlay.h"
+#include "Widgets/Text/STextBlock.h"
+#include "Widgets/SBoxPanel.h"
+
+#define LOCTEXT_NAMESPACE "UnrealMultiUserUI.SConcertClientBrowserItem"
+
+void UE::MultiUserServer::SConcertBrowserItem::Construct(const FArguments& InArgs, TSharedRef<IConcertBrowserItem> InClientItem)
+{
+	Item = MoveTemp(InClientItem);
+	HighlightText = InArgs._HighlightText;
+
+	ChildSlot
+	.Padding(FMargin(0.0f, 0.0f, 4.0f, 4.0f))
+	[
+		// Shadow behind thumbnail
+		SNew(SBorder)
+		.BorderImage(FConcertServerStyle::Get().GetBrush("Concert.Clients.DropShadow"))
+		.Padding(FMargin(0.0f, 0.0f, 5.0f, 5.0f))
+		[
+			// Change outside of thumbnail depending on hover state; lighter than inside
+			SNew(SBorder)
+			.BorderImage(this, &SConcertBrowserItem::GetBackgroundImage)
+			.Padding(2.f)
+			[
+				// Inside of thumbnail is darker
+				SNew(SBorder)
+				.BorderImage(FConcertServerStyle::Get().GetBrush("Concert.Clients.ThumbnailTitle"))
+				.Padding(2.f)
+				[
+					SNew(SVerticalBox)
+
+					+SVerticalBox::Slot()
+					.AutoHeight()
+					.VAlign(VAlign_Top)
+					[
+						CreateHeader()	
+					]
+					
+					+SVerticalBox::Slot()
+					.FillHeight(1.f)
+					.Padding(-2.f, 5.f, -2.f, 0.f)
+					[
+						CreateContentArea()
+					]
+
+					+SVerticalBox::Slot()
+					.AutoHeight()
+					.Padding(0.f, 5.f, 0.f, 0.f)
+					.HAlign(HAlign_Fill)
+					.VAlign(VAlign_Bottom)
+					[
+						CreateStats()
+					]
+
+					+SVerticalBox::Slot()
+					.AutoHeight()
+					.Padding(-2.f, 5.f, -2.f, -2.f)
+					.VAlign(VAlign_Bottom)
+					[
+						CreateFooter()
+					]
+				]
+			]
+		]
+	];
+}
+
+TSharedRef<SWidget> UE::MultiUserServer::SConcertBrowserItem::CreateHeader()
+{
+	return SAssignNew(ClientName, STextBlock)
+		.Font(FConcertServerStyle::Get().GetFontStyle("Concert.Clients.ClientNameTileFont"))
+		.Text_Lambda([this](){ return FText::FromString(Item->GetDisplayName()); })
+		.ToolTipText_Lambda([this](){ return Item->GetToolTip(); })
+		.HighlightText_Lambda([this](){ return *HighlightText; })
+		.ColorAndOpacity(FColor::White);
+}
+
+TSharedRef<SWidget> UE::MultiUserServer::SConcertBrowserItem::CreateContentArea()
+{
+	return SNew(SOverlay)
+		+SOverlay::Slot()
+		[
+			SNew(SClientNetworkGraphs, Item->GetTransferStatistics())
+			.Visibility_Lambda([this](){ return Item->GetDisplayMode() == EConcertBrowserItemDisplayMode::NetworkGraph ? EVisibility::Visible : EVisibility::Collapsed; })
+		]
+		+SOverlay::Slot().Padding(4.f, 0.f)
+		[
+			SNew(SClientOutboundTransferStatTable, Item->GetTransferStatistics())
+			.Visibility_Lambda([this](){ return Item->GetDisplayMode() == EConcertBrowserItemDisplayMode::OutboundSegementTable ? EVisibility::Visible : EVisibility::Collapsed; })
+		]
+		+SOverlay::Slot().Padding(4.f, 0.f)
+		[
+			SNew(SClientInboundTransferStatTable, Item->GetTransferStatistics())
+			.Visibility_Lambda([this](){ return Item->GetDisplayMode() == EConcertBrowserItemDisplayMode::InboundSegmentTable ? EVisibility::Visible : EVisibility::Collapsed; })
+		];
+}
+
+TSharedRef<SWidget> UE::MultiUserServer::SConcertBrowserItem::CreateStats()
+{
+	return SNew(SScaleBox)
+		.Stretch(EStretch::ScaleToFit)
+		.HAlign(HAlign_Fill)
+		[
+			SAssignNew(NetworkStats, SClientNetworkStats)
+			.HighlightText(HighlightText)
+			.NetworkStatistics_Lambda([this](){ return Item->GetLatestNetworkStatistics(); })
+		];
+}
+
+TSharedRef<SWidget> UE::MultiUserServer::SConcertBrowserItem::CreateFooter()
+{
+	return SNew(SBorder)
+		.BorderImage(FConcertServerStyle::Get().GetBrush("Concert.Clients.ThumbnailFooter"))
+		[
+			SNew(SHorizontalBox)
+
+			// Online / offline indicator
+			+SHorizontalBox::Slot()
+			.HAlign(HAlign_Left)
+			.Padding(2.f)
+			[
+				SNew(SImage)
+				.Image(FAppStyle::GetBrush("Icons.FilledCircle"))
+				.ColorAndOpacity_Lambda([this]()
+				{
+					return Item->IsOnline()
+						? FStyleColors::AccentGreen
+						: FStyleColors::AccentGray;
+				})
+				.ToolTipText_Lambda([this]()
+				{
+					return Item->IsOnline()
+						? LOCTEXT("ConnectionIndicator.Online", "Connected")
+						: LOCTEXT("ConnectionIndicator.Offline", "Not reachable");
+				})
+			]
+
+			// IP address
+			+SHorizontalBox::Slot()
+			.HAlign(HAlign_Right)
+			.Padding(2.f)
+			[
+				SAssignNew(ClientIP4, STextBlock)
+				.ColorAndOpacity(FColor::White)
+				.HighlightText_Lambda([this](){ return *HighlightText; })
+				.Text_Lambda([this]()
+				{
+					return FText::FromString(NetworkStatistics::FormatIPv4AsString(Item->GetLatestNetworkStatistics()));
+				})
+			]
+		];
+}
+
+const FSlateBrush* UE::MultiUserServer::SConcertBrowserItem::GetBackgroundImage() const
+{
+	if (IsHovered())
+	{
+		return FConcertServerStyle::Get().GetBrush("Concert.Clients.ThumbnailAreaHoverBackground");
+	}
+	return FConcertServerStyle::Get().GetBrush("Concert.Clients.ThumbnailAreaBackground");
+}
+
+#undef LOCTEXT_NAMESPACE
\ No newline at end of file
diff --git a/Engine/Plugins/Developer/Concert/ConcertApp/MultiUserServer/Source/MultiUserServer/Private/Widgets/Clients/Browser/SConcertBrowserItem.h b/Engine/Plugins/Developer/Concert/ConcertApp/MultiUserServer/Source/MultiUserServer/Private/Widgets/Clients/Browser/SConcertBrowserItem.h
new file mode 100644
--- /dev/null
+++ b/Engine/Plugins/Developer/Concert/ConcertApp/MultiUserServer/Source/MultiUserServer/Private/Widgets/Clients/Browser/SConcertBrowserItem.h
@@ -0,0 +1,49 @@
+// Copyright Epic Games, Inc. All Rights Reserved.
+
+#pragma once
+
+#include "CoreMinimal.h"
+#include "IMessageContext.h"
+#include "Models/IClientTransferStatisticsModel.h"
+#include "Widgets/DeclarativeSyntaxSupport.h"
+#include "Widgets/SCompoundWidget.h"
+
+class FMenuBuilder;
+class STextBlock;
+
+namespace UE::MultiUserServer
+{
+	class SClientNetworkStats;
+	class IClientNetworkStatisticsModel;
+	class IConcertBrowserItem;
+	
+	class SConcertBrowserItem : public SCompoundWidget
+	{
+	public:
+
+		SLATE_BEGIN_ARGS(SConcertBrowserItem) {}
+			SLATE_ARGUMENT(TSharedPtr<FText>, HighlightText)
+		SLATE_END_ARGS()
+		
+		void Construct(const FArguments& InArgs, TSharedRef<IConcertBrowserItem> InClientItem);
+
+	private:
+
+		/** What we're displaying */
+		TSharedPtr<IConcertBrowserItem> Item;
+
+		/** The text to highlight */
+		TSharedPtr<FText> HighlightText;
+
+		TSharedPtr<STextBlock> ClientName;
+		TSharedPtr<SClientNetworkStats> NetworkStats;
+		TSharedPtr<STextBlock> ClientIP4;
+		
+		TSharedRef<SWidget> CreateHeader();
+		TSharedRef<SWidget> CreateContentArea();
+		TSharedRef<SWidget> CreateStats();
+		TSharedRef<SWidget> CreateFooter();
+
+		const FSlateBrush* GetBackgroundImage() const;
+	};
+}
diff --git a/Engine/Plugins/Developer/Concert/ConcertApp/MultiUserServer/Source/MultiUserServer/Private/Widgets/Clients/Browser/SConcertClientBrowser.cpp b/Engine/Plugins/Developer/Concert/ConcertApp/MultiUserServer/Source/MultiUserServer/Private/Widgets/Clients/Browser/SConcertClientBrowser.cpp
deleted file mode 100644
diff --git a/Engine/Plugins/Developer/Concert/ConcertApp/MultiUserServer/Source/MultiUserServer/Private/Widgets/Clients/Browser/SConcertClientBrowser.h b/Engine/Plugins/Developer/Concert/ConcertApp/MultiUserServer/Source/MultiUserServer/Private/Widgets/Clients/Browser/SConcertClientBrowser.h
deleted file mode 100644
diff --git a/Engine/Plugins/Developer/Concert/ConcertApp/MultiUserServer/Source/MultiUserServer/Private/Widgets/Clients/Browser/SConcertClientBrowserItem.cpp b/Engine/Plugins/Developer/Concert/ConcertApp/MultiUserServer/Source/MultiUserServer/Private/Widgets/Clients/Browser/SConcertClientBrowserItem.cpp
deleted file mode 100644
diff --git a/Engine/Plugins/Developer/Concert/ConcertApp/MultiUserServer/Source/MultiUserServer/Private/Widgets/Clients/Browser/SConcertClientBrowserItem.h b/Engine/Plugins/Developer/Concert/ConcertApp/MultiUserServer/Source/MultiUserServer/Private/Widgets/Clients/Browser/SConcertClientBrowserItem.h
deleted file mode 100644
diff --git a/Engine/Plugins/Developer/Concert/ConcertApp/MultiUserServer/Source/MultiUserServer/Private/Widgets/Clients/Browser/SConcertNetworkBrowser.cpp b/Engine/Plugins/Developer/Concert/ConcertApp/MultiUserServer/Source/MultiUserServer/Private/Widgets/Clients/Browser/SConcertNetworkBrowser.cpp
new file mode 100644
--- /dev/null
+++ b/Engine/Plugins/Developer/Concert/ConcertApp/MultiUserServer/Source/MultiUserServer/Private/Widgets/Clients/Browser/SConcertNetworkBrowser.cpp
@@ -0,0 +1,458 @@
+// Copyright Epic Games, Inc. All Rights Reserved.
+
+#include "SConcertNetworkBrowser.h"
+
+#include "ClientBrowserItem.h"
+#include "ConcertServerStyle.h"
+#include "MultiUserServerModule.h"
+#include "Models/IClientBrowserModel.h"
+#include "Models/IClientNetworkStatisticsModel.h"
+#include "SConcertBrowserItem.h"
+#include "Algo/Accumulate.h"
+#include "Algo/AllOf.h"
+
+#include "Algo/AnyOf.h"
+#include "Algo/Count.h"
+#include "Dialog/SMessageDialog.h"
+#include "Framework/MultiBox/MultiBoxBuilder.h"
+#include "Widgets/Input/SCheckBox.h"
+#include "Widgets/Input/SComboButton.h"
+#include "Widgets/Input/SSearchBox.h"
+#include "Widgets/SBoxPanel.h"
+#include "Widgets/Text/STextBlock.h"
+#include "Widgets/Views/STileView.h"
+#include "Window/ModalWindowManager.h"
+
+#define LOCTEXT_NAMESPACE "UnrealMultiUserUI.SConcertClientBrowser"
+
+void UE::MultiUserServer::SConcertNetworkBrowser::Construct(const FArguments& InArgs, TSharedRef<IClientBrowserModel> InBrowserModel)
+{
+	BrowserModel = MoveTemp(InBrowserModel);
+	
+	HighlightText = MakeShared<FText>();
+	SessionFilter = MakeShared<FClientTextFilter>(FClientTextFilter::FItemToStringArray::CreateSP(this, &SConcertNetworkBrowser::GenerateSearchTerms));
+	SessionFilter->OnChanged().AddSP(this, &SConcertNetworkBrowser::UpdateTileViewFromAllowedSessions);
+
+	OnClientDoubleClicked = InArgs._OnClientDoubleClicked;
+
+	ChildSlot
+	[
+		SNew(SVerticalBox)
+
+		+SVerticalBox::Slot()
+		.AutoHeight()
+		[
+			CreateSearchArea(InArgs)
+		]
+
+		+SVerticalBox::Slot()
+		.FillHeight(1.f)
+		.Padding(5.f)
+		[
+			SNew(SOverlay)
+
+			+SOverlay::Slot()
+			[
+				CreateTileView()
+			]
+
+			+SOverlay::Slot()
+			[
+				SNew(SBox)
+				.HAlign(HAlign_Center)
+				.VAlign(VAlign_Center)
+				.Visibility_Lambda([this](){ return DisplayedClients.Num() == 0 ? EVisibility::Visible : EVisibility::Hidden; })
+				[
+					SNew(STextBlock)
+					.Text(this, &SConcertNetworkBrowser::GetErrorMessageText)
+				]
+			]
+		]
+	];
+
+	BrowserModel->OnSessionCreated().AddSP(this, &SConcertNetworkBrowser::OnSessionCreated);
+	BrowserModel->OnSessionDestroyed().AddSP(this, &SConcertNetworkBrowser::OnSessionDestroyed);
+	BrowserModel->OnClientListChanged().AddSP(this, &SConcertNetworkBrowser::OnClientListChanged);
+
+	AllowAllSessions();
+}
+
+void UE::MultiUserServer::SConcertNetworkBrowser::ShowOnlyClientsFromSession(const FGuid& SessionId)
+{
+	AllowedSessions = { SessionId };
+	UpdateTileViewFromAllowedSessions();
+}
+
+TSharedRef<SWidget> UE::MultiUserServer::SConcertNetworkBrowser::CreateSearchArea(const FArguments& InArgs)
+{
+	return SNew(SHorizontalBox)
+	
+		+SHorizontalBox::Slot()
+		.AutoWidth()
+		[
+			SNew(SComboButton)
+			.OnGetMenuContent(this, &SConcertNetworkBrowser::MakeSessionOption)
+			.ButtonContent()
+			[
+				SNew(STextBlock)
+				.Text_Lambda([this]()
+				{
+					return AllowedSessions.Num() == BrowserModel->GetSessions().Num()
+						? LOCTEXT("MessageActionFilter.Selection.All", "All")
+						: FText::FromString(FString::FromInt(AllowedSessions.Num()));
+				})
+			]
+		]
+
+		+SHorizontalBox::Slot()
+		.AutoWidth()
+		.FillWidth(0.65f)
+		.Padding(4.f, 0.f)
+		.VAlign(VAlign_Fill)
+		[
+			SNew(SSearchBox)
+			.OnTextChanged_Lambda([this](const FText& SearchText)
+			{
+				*HighlightText = SearchText;
+				SessionFilter->SetRawFilterText(SearchText);
+			})
+		]
+
+		+SHorizontalBox::Slot()
+		.AutoWidth()
+		[
+			InArgs._RightOfSearch.Widget
+		]
+
+		+SHorizontalBox::Slot()
+		.AutoWidth()
+		.HAlign(HAlign_Right)
+		[
+			CreateKeepDisconnectedClients()
+		];
+}
+
+TSharedRef<SWidget> UE::MultiUserServer::SConcertNetworkBrowser::CreateKeepDisconnectedClients()
+{
+	return SNew(SHorizontalBox)
+		.ToolTipText(LOCTEXT("KeepDisconnectedClients.Tooltip", "Whether to keep clients that have disconnected in memory. This may be useful in unstable networks when you want to analyse why clients keep disconnecting."))
+		
+		+SHorizontalBox::Slot()
+		.AutoWidth()
+		.VAlign(VAlign_Center)
+		[
+			SNew(STextBlock)
+			.Text(LOCTEXT("KeepDisconnectedClients.Label", "Keep Disconnected"))
+			]
+
+		+SHorizontalBox::Slot()
+		.AutoWidth()
+		.Padding(4.f, 0.f, 0.f, 0.f)
+		[
+			SNew(SCheckBox)
+			.IsChecked_Lambda([this](){ return BrowserModel->ShouldKeepClientsAfterDisconnect() ? ECheckBoxState::Checked : ECheckBoxState::Unchecked; })
+			.OnCheckStateChanged_Lambda([this](ECheckBoxState NewState)
+			{
+				const bool bKeepDisconnected = NewState == ECheckBoxState::Checked;
+				if (bKeepDisconnected)
+				{
+					BrowserModel->SetKeepClientsAfterDisconnect(true);
+					return;
+				}
+				
+				const int32 NumDisconnected = Algo::CountIf(BrowserModel->GetItems(), [](const TSharedPtr<FClientBrowserItem>& Item) { return Item->IsDisconnected(); });
+				if (NumDisconnected == 0)
+				{
+					BrowserModel->SetKeepClientsAfterDisconnect(false);
+					return;
+				}
+				
+				const TSharedRef<SMessageDialog> Dialog = SNew(SMessageDialog)
+					.Title(LOCTEXT("RemoveDisconnectedClients.Title", "Remove disconnected clients?"))
+					.Icon(FAppStyle::Get().GetBrush("Icons.WarningWithColor.Large"))
+					.Message(FText::Format(LOCTEXT("RemoveDisconnectedClients.MessageFmt", "There are {0} disconnected clients. If you proceed, these clients will be removed from the session browser; opened log tabs will remain open.\nProceed?"), NumDisconnected))
+					.UseScrollBox(false)
+					.Buttons({
+						SMessageDialog::FButton(LOCTEXT("RemoveButton", "Remove"))
+							.SetOnClicked(FSimpleDelegate::CreateLambda([this]()
+							{
+								BrowserModel->SetKeepClientsAfterDisconnect(false);
+							})),
+						SMessageDialog::FButton(LOCTEXT("CancelButton", "Keep"))
+							.SetPrimary(true)
+							.SetFocus()
+					});
+				FConcertServerUIModule::Get().GetModalWindowManager()->ShowFakeModalWindow(Dialog);
+			})
+		]
+	;
+}
+
+TSharedRef<SWidget> UE::MultiUserServer::SConcertNetworkBrowser::CreateTileView()
+{
+	constexpr int32 Height = 270;
+	constexpr int32 Width = 270;
+	return SAssignNew(TileView, STileView<TSharedPtr<FClientBrowserItem>>)
+		.SelectionMode(ESelectionMode::Multi)
+		.ListItemsSource(&DisplayedClients)
+		.OnGenerateTile(this, &SConcertNetworkBrowser::MakeTileViewWidget)
+		.OnContextMenuOpening(this, &SConcertNetworkBrowser::OnGetContextMenuContent)
+		.OnMouseButtonDoubleClick(this, &SConcertNetworkBrowser::OnListMouseButtonDoubleClick)
+		.ItemHeight(Height)
+		.ItemWidth(Width)
+	;
+}
+
+void UE::MultiUserServer::SConcertNetworkBrowser::OnSessionCreated(const FGuid& SessionId)
+{
+	if (bShowAllSessions)
+	{
+		AllowedSessions.Add(SessionId);
+		UpdateTileViewFromAllowedSessions();
+	}
+}
+
+void UE::MultiUserServer::SConcertNetworkBrowser::OnSessionDestroyed(const FGuid& SessionId)
+{
+	if (AllowedSessions.Contains(SessionId))
+	{
+		AllowedSessions.Remove(SessionId);
+		UpdateTileViewFromAllowedSessions();
+	}
+}
+
+void UE::MultiUserServer::SConcertNetworkBrowser::OnClientListChanged(TSharedPtr<FClientBrowserItem> Item, IClientBrowserModel::EClientUpdateType UpdateType)
+{
+	switch (UpdateType)
+	{
+	case IClientBrowserModel::EClientUpdateType::Added:
+		if (PassesFilter(Item))
+		{
+			DisplayedClients.Add(Item);
+		}
+		break;
+	case IClientBrowserModel::EClientUpdateType::Removed:
+		DisplayedClients.Remove(Item);
+		break;
+	default:
+		checkNoEntry();
+	}
+
+	TileView->RequestListRefresh();
+}
+
+TSharedRef<SWidget> UE::MultiUserServer::SConcertNetworkBrowser::SConcertNetworkBrowser::MakeSessionOption()
+{
+	FMenuBuilder MenuBuilder(true, nullptr);
+
+	MenuBuilder.AddMenuEntry(
+		LOCTEXT("SelectAll", "All"),
+		FText::GetEmpty(),
+		FSlateIcon(),
+		FUIAction(
+			FExecuteAction::CreateLambda([this]()
+			{
+				if (BrowserModel->GetSessions().Num() == AllowedSessions.Num())
+				{
+					DisallowAllSessions();
+				}
+				else
+				{
+					AllowAllSessions();
+				}
+			}),
+			FCanExecuteAction::CreateLambda([] { return true; }),
+			FIsActionChecked::CreateLambda([this](){ return BrowserModel->GetSessions().Num() == AllowedSessions.Num(); })),
+		NAME_None,
+		EUserInterfaceActionType::Check
+	);
+	MenuBuilder.AddMenuEntry(
+		LOCTEXT("SessionlessEndpoints.Label", "Show Sessionless clients"),
+		LOCTEXT("SessionlessEndpoints.Tooltip", "Whether to show clients that are only discovering available sessions"),
+		FSlateIcon(),
+		FUIAction(
+			FExecuteAction::CreateLambda([this]()
+			{
+				bShowSessionlessClients = !bShowSessionlessClients;
+				UpdateTileViewFromAllowedSessions();
+			}),
+			FCanExecuteAction::CreateLambda([] { return true; }),
+			FIsActionChecked::CreateLambda([this](){ return bShowSessionlessClients; })),
+		NAME_None,
+		EUserInterfaceActionType::Check
+	);
+	MenuBuilder.AddSeparator();
+	for (const FGuid& SessionId : BrowserModel->GetSessions())
+	{
+		MenuBuilder.AddMenuEntry(
+			FText::FromString(BrowserModel->GetSessionInfo(SessionId)->SessionName),
+			FText::GetEmpty(),
+			FSlateIcon(),
+			FUIAction(
+				FExecuteAction::CreateLambda([this, SessionId]()
+				{
+					if (AllowedSessions.Contains(SessionId))
+					{
+						AllowedSessions.Remove(SessionId);
+					}
+					else
+					{
+						AllowedSessions.Add(SessionId);
+					}
+					UpdateTileViewFromAllowedSessions();
+				}),
+				FCanExecuteAction::CreateLambda([] { return true; }),
+				FIsActionChecked::CreateLambda([this, SessionId]{ return AllowedSessions.Contains(SessionId); })),
+			NAME_None,
+			EUserInterfaceActionType::Check
+		);
+	}
+	return MenuBuilder.MakeWidget();
+}
+
+FText UE::MultiUserServer::SConcertNetworkBrowser::GetErrorMessageText() const
+{
+	if (BrowserModel->GetSessions().Num() == 0)
+	{
+		return LOCTEXT("NoLiveSessions", "No live sessions");
+	}
+	
+	if (AllowedSessions.Num() == 0)
+	{
+		return LOCTEXT("NoSessionsSelected", "All live sessions filtered out");
+	}
+
+	const bool bAtLeastOneClient = BrowserModel->GetItems().Num() > 0;
+	if (bAtLeastOneClient)
+	{
+		return LOCTEXT("AllFilteredOut", "All results have been filtered. Try changing your active filters above.");
+	}
+	
+	return LOCTEXT("NoClients", "No known clients");
+}
+
+TSharedRef<ITableRow> UE::MultiUserServer::SConcertNetworkBrowser::MakeTileViewWidget(TSharedPtr<FClientBrowserItem> ClientItem, const TSharedRef<STableViewBase>& OwnerTable)
+{
+	return SNew(STableRow<TSharedPtr<FClientBrowserItem>>, OwnerTable )
+		.Padding(2.f)
+		.Style(FConcertServerStyle::Get(), "Concert.Clients.TileTableRow")
+		.Content()
+		[
+			SNew(SConcertBrowserItem, ClientItem.ToSharedRef())
+			 .HighlightText(HighlightText)
+		];
+}
+
+TSharedPtr<SWidget> UE::MultiUserServer::SConcertNetworkBrowser::OnGetContextMenuContent()
+{
+	FMenuBuilder MenuBuilder(true, nullptr);
+
+	MenuBuilder.AddMenuEntry(
+		LOCTEXT("OpenLog", "Open log"),
+		LOCTEXT("OpenLogTooltip", "Opens a new tab in which you can filter log events related to this client"),
+		FSlateIcon(),
+		FUIAction(
+			FExecuteAction::CreateLambda([this]()
+			{
+				TArray<TSharedPtr<FClientBrowserItem>> Items = TileView->GetSelectedItems();
+				for (const TSharedPtr<FClientBrowserItem>& Item : Items)
+				{
+					OnListMouseButtonDoubleClick(Item);
+				}
+			}),
+			FCanExecuteAction::CreateLambda([] { return true; }),
+			FIsActionChecked()),
+		NAME_None,
+		EUserInterfaceActionType::Button
+		);
+	
+	MenuBuilder.AddSeparator();
+	AddDisplayModeEntry(MenuBuilder, EConcertBrowserItemDisplayMode::NetworkGraph,
+		LOCTEXT("DisplayMode.NetworkGraph.Title", "Network graph"),
+		LOCTEXT("DisplayMode.NetworkGraph.Tooltip", "Show the up and down stream network traffic on a graph")
+		);
+	AddDisplayModeEntry(MenuBuilder, EConcertBrowserItemDisplayMode::OutboundSegementTable,
+		LOCTEXT("DisplayMode.OutboundSegementTable.Title", "Outbound segment table"),
+		LOCTEXT("DisplayMode.SegementTable.Tooltip", "A table displaying the messaging protocol's outbound segments' MessageId, Sent, Acked and Size data in realtime.")
+		);
+	AddDisplayModeEntry(MenuBuilder, EConcertBrowserItemDisplayMode::InboundSegmentTable,
+		LOCTEXT("DisplayMode.InboundSegementTable.Title", "Inbound Segment table"),
+		LOCTEXT("DisplayMode.SegementTable.Tooltip", "A table displaying the messaging protocol's inbound segments' MessageId, Received and Size data in realtime.")
+		);
+	return MenuBuilder.MakeWidget();
+}
+
+void UE::MultiUserServer::SConcertNetworkBrowser::AddDisplayModeEntry(FMenuBuilder& MenuBuilder, EConcertBrowserItemDisplayMode DisplayMode, FText Title, FText Tooltip) const
+{
+	MenuBuilder.AddMenuEntry(
+		Title,
+		Tooltip,
+		FSlateIcon(),
+		FUIAction(
+			FExecuteAction::CreateLambda([this, DisplayMode]()
+			{
+				TArray<TSharedPtr<FClientBrowserItem>> Items = TileView->GetSelectedItems();
+				for (const TSharedPtr<FClientBrowserItem>& Item : Items)
+				{
+					Item->SetDisplayMode(DisplayMode);
+				}
+			}),
+			FCanExecuteAction::CreateLambda([] { return true; }),
+			FIsActionChecked::CreateLambda([this, DisplayMode]()
+			{
+				return Algo::AllOf(TileView->GetSelectedItems(), [this, DisplayMode](const TSharedPtr<FClientBrowserItem>& Item)
+				{
+					return Item->GetDisplayMode() == DisplayMode;
+				});
+			})),
+		NAME_None,
+		EUserInterfaceActionType::Check
+	);
+}
+
+void UE::MultiUserServer::SConcertNetworkBrowser::OnListMouseButtonDoubleClick(TSharedPtr<FClientBrowserItem> ClientItem)
+{
+	OnClientDoubleClicked.ExecuteIfBound(ClientItem->GetMessageNodeId());
+}
+
+void UE::MultiUserServer::SConcertNetworkBrowser::SConcertNetworkBrowser::AllowAllSessions()
+{
+	AllowedSessions = BrowserModel->GetSessions();
+	UpdateTileViewFromAllowedSessions();
+}
+
+void UE::MultiUserServer::SConcertNetworkBrowser::DisallowAllSessions()
+{
+	AllowedSessions.Reset();
+	UpdateTileViewFromAllowedSessions();
+}
+
+void UE::MultiUserServer::SConcertNetworkBrowser::UpdateTileViewFromAllowedSessions()
+{
+	bShowAllSessions = AllowedSessions.Num() == BrowserModel->GetSessions().Num();
+	DisplayedClients.Empty();
+	
+	for (const TSharedPtr<FClientBrowserItem>& Item : BrowserModel->GetItems())
+	{
+		if (PassesFilter(Item))
+		{
+			DisplayedClients.Add(Item);
+		}
+	}
+	TileView->RequestListRefresh();
+}
+
+bool UE::MultiUserServer::SConcertNetworkBrowser::PassesFilter(const TSharedPtr<FClientBrowserItem>& Client) const
+{
+	const bool bIsSessionAllowed = Client->GetCurrentSession().IsSet() && AllowedSessions.Contains(*Client->GetCurrentSession());
+	const bool bAllowBaseOnSession = bIsSessionAllowed || (bShowSessionlessClients && !Client->GetCurrentSession().IsSet());
+	return bAllowBaseOnSession && SessionFilter->PassesFilter(Client);
+}
+
+void UE::MultiUserServer::SConcertNetworkBrowser::GenerateSearchTerms(const TSharedPtr<FClientBrowserItem>& Client, TArray<FString>& SearchTerms) const
+{
+	Client->AppendSearchTerms(SearchTerms);
+}
+
+
+#undef LOCTEXT_NAMESPACE
diff --git a/Engine/Plugins/Developer/Concert/ConcertApp/MultiUserServer/Source/MultiUserServer/Private/Widgets/Clients/Browser/SConcertNetworkBrowser.h b/Engine/Plugins/Developer/Concert/ConcertApp/MultiUserServer/Source/MultiUserServer/Private/Widgets/Clients/Browser/SConcertNetworkBrowser.h
new file mode 100644
--- /dev/null
+++ b/Engine/Plugins/Developer/Concert/ConcertApp/MultiUserServer/Source/MultiUserServer/Private/Widgets/Clients/Browser/SConcertNetworkBrowser.h
@@ -0,0 +1,105 @@
+// Copyright Epic Games, Inc. All Rights Reserved.
+
+#pragma once
+
+#include "CoreMinimal.h"
+#include "SConcertBrowserItem.h"
+#include "Misc/TextFilter.h"
+#include "Models/IClientBrowserModel.h"
+#include "Widgets/DeclarativeSyntaxSupport.h"
+#include "Widgets/SCompoundWidget.h"
+
+class FConcertClientsTabController;
+class FMenuBuilder;
+class STableViewBase;
+class ITableRow;
+template<typename T> class STileView;
+
+enum class EConcertClientStatus : uint8;
+struct FConcertSessionClientInfo;
+
+namespace UE::MultiUserServer
+{
+	class FClientBrowserItem;
+	class IClientNetworkStatisticsModel;
+	class IClientBrowserModel;
+	enum class EConcertBrowserItemDisplayMode : uint8;
+	
+	/**
+	 * Presents a view of client thumbnails, which include important statistics.
+	 * Clients can be right-clicked, opening a context menu, and double clicked, opening a new tab with a log view.
+	 */
+	class SConcertNetworkBrowser : public SCompoundWidget
+	{
+	public:
+
+		DECLARE_DELEGATE_OneParam(FOnClientDoubleClicked, const FGuid& /*EndpointId*/);
+
+		SLATE_BEGIN_ARGS(SConcertNetworkBrowser) {}
+			/** Extension point to the right of the search bar */
+			SLATE_NAMED_SLOT(FArguments, RightOfSearch)
+			SLATE_EVENT(FOnClientDoubleClicked, OnClientDoubleClicked)
+		SLATE_END_ARGS()
+
+		void Construct(const FArguments& InArgs, TSharedRef<IClientBrowserModel> InBrowserModel);
+
+		/** Shows only the clients connected to the given session ID */
+		void ShowOnlyClientsFromSession(const FGuid& SessionId);
+		
+	private:
+
+		using FSessionId = FGuid;
+		using FMessagingNodeId = FGuid;
+		using FClientTextFilter = TTextFilter<const TSharedPtr<FClientBrowserItem>&>;
+
+		/** Retrieves clients and live sessions */
+		TSharedPtr<IClientBrowserModel> BrowserModel;
+
+		/** Contains the guid of every session we're allowed to display */
+		TSet<FSessionId> AllowedSessions;
+		/** Should all sessions be shown? */
+		bool bShowAllSessions = true;
+		/** Should admin endpoints be shown? */
+		bool bShowSessionlessClients = true;
+
+		/** Source array for TileView - fitlered version of IClientBrowserModel::GetItems */
+		TArray<TSharedPtr<FClientBrowserItem>> DisplayedClients;
+		/** Visualizes all the items */
+		TSharedPtr<STileView<TSharedPtr<FClientBrowserItem>>> TileView;
+
+		// Filtering
+		TSharedPtr<FText> HighlightText;
+		TSharedPtr<FClientTextFilter> SessionFilter;
+		
+		FOnClientDoubleClicked OnClientDoubleClicked;
+
+		TSharedRef<SWidget> CreateSearchArea(const FArguments& InArgs);
+		TSharedRef<SWidget> CreateKeepDisconnectedClients();
+		TSharedRef<SWidget> CreateTileView();
+
+		// Model events
+		void OnSessionCreated(const FGuid& SessionId);
+		void OnSessionDestroyed(const FGuid& SessionId);
+		void OnClientListChanged(TSharedPtr<FClientBrowserItem> Item, IClientBrowserModel::EClientUpdateType UpdateType);
+
+		// Combo button
+		TSharedRef<SWidget> MakeSessionOption();
+
+		FText GetErrorMessageText() const;
+		
+		// TileView events
+		TSharedRef<ITableRow> MakeTileViewWidget(TSharedPtr<FClientBrowserItem> ClientItem, const TSharedRef<STableViewBase>& OwnerTable);
+		TSharedPtr<SWidget> OnGetContextMenuContent();
+		void AddDisplayModeEntry(FMenuBuilder& MenuBuilder, EConcertBrowserItemDisplayMode DisplayMode, FText Title, FText Tooltip) const;
+		void OnListMouseButtonDoubleClick(TSharedPtr<FClientBrowserItem> ClientItem);
+
+		// Filtering
+		void AllowAllSessions();
+		void DisallowAllSessions();
+		void UpdateTileViewFromAllowedSessions();
+		bool PassesFilter(const TSharedPtr<FClientBrowserItem>& Client) const;
+		void GenerateSearchTerms(const TSharedPtr<FClientBrowserItem>& Client, TArray<FString>& SearchTerms) const;
+	};
+}
+
+
diff --git a/Engine/Plugins/Developer/Concert/ConcertApp/MultiUserServer/Source/MultiUserServer/Private/Widgets/Clients/SConcertClientsTabView.cpp b/Engine/Plugins/Developer/Concert/ConcertApp/MultiUserServer/Source/MultiUserServer/Private/Widgets/Clients/SConcertClientsTabView.cpp
--- a/Engine/Plugins/Developer/Concert/ConcertApp/MultiUserServer/Source/MultiUserServer/Private/Widgets/Clients/SConcertClientsTabView.cpp
+++ b/Engine/Plugins/Developer/Concert/ConcertApp/MultiUserServer/Source/MultiUserServer/Private/Widgets/Clients/SConcertClientsTabView.cpp
@@ -9,7 +9,7 @@
 #include "Logging/Source/GlobalLogSource.h"
 #include "Logging/Util/ConcertLogTokenizer.h"
 #include "Util/EndpointToUserNameCache.h"
-#include "Widgets/Clients/Browser/SConcertClientBrowser.h"
+#include "Widgets/Clients/Browser/SConcertNetworkBrowser.h"
 #include "Widgets/Clients/Browser/Models/ClientBrowserModel.h"
 #include "Widgets/Clients/Browser/Models/ClientNetworkStatisticsModel.h"
 #include "Widgets/Clients/Logging/SConcertTransportLog.h"
@@ -147,9 +147,8 @@ TSharedRef<SDockTab> SConcertClientsTabView::SpawnClientBrowserTab(const FSpawnT
 		.Label(LOCTEXT("ClientBrowserTabLabel", "Clients"))
 		.TabRole(PanelTab)
 		[
-			SAssignNew(ClientBrowser, UE::MultiUserServer::SConcertClientBrowser,
-				MakeShared<UE::MultiUserServer::FClientBrowserModel>(Server->GetConcertServer(), ClientInfoCache.ToSharedRef()),
-				MakeShared<UE::MultiUserServer::FClientNetworkStatisticsModel>())
+			SAssignNew(ClientBrowser, UE::MultiUserServer::SConcertNetworkBrowser,
+				MakeShared<UE::MultiUserServer::FClientBrowserModel>(Server->GetConcertServer(), ClientInfoCache.ToSharedRef(), MakeShared<UE::MultiUserServer::FClientNetworkStatisticsModel>()))
 			.RightOfSearch()
 			[
 				CreateOpenGlobalLogButton()
diff --git a/Engine/Plugins/Developer/Concert/ConcertApp/MultiUserServer/Source/MultiUserServer/Private/Widgets/Clients/SConcertClientsTabView.h b/Engine/Plugins/Developer/Concert/ConcertApp/MultiUserServer/Source/MultiUserServer/Private/Widgets/Clients/SConcertClientsTabView.h
--- a/Engine/Plugins/Developer/Concert/ConcertApp/MultiUserServer/Source/MultiUserServer/Private/Widgets/Clients/SConcertClientsTabView.h
+++ b/Engine/Plugins/Developer/Concert/ConcertApp/MultiUserServer/Source/MultiUserServer/Private/Widgets/Clients/SConcertClientsTabView.h
@@ -10,7 +10,7 @@ namespace UE
 {
 	namespace MultiUserServer
 	{
-		class SConcertClientBrowser;
+		class SConcertNetworkBrowser;
 	}
 }
 
@@ -62,7 +62,7 @@ private:
 	/** Used by various systems to convert logs to text */
 	TSharedPtr<FConcertLogTokenizer> LogTokenizer;
 
-	TSharedPtr<UE::MultiUserServer::SConcertClientBrowser> ClientBrowser;
+	TSharedPtr<UE::MultiUserServer::SConcertNetworkBrowser> ClientBrowser;
 	
 	/** Used to overlay EnableLoggingPrompt over the tabs */
 	TSharedPtr<SOverlay> EnableLoggingPromptOverlay;
