Fix localisation warning

jira none

diff --git a/Engine/Plugins/Developer/Concert/ConcertApp/MultiUserServer/Source/MultiUserServer/Private/Widgets/Clients/Browser/SConcertNetworkBrowser.cpp b/Engine/Plugins/Developer/Concert/ConcertApp/MultiUserServer/Source/MultiUserServer/Private/Widgets/Clients/Browser/SConcertNetworkBrowser.cpp
--- a/Engine/Plugins/Developer/Concert/ConcertApp/MultiUserServer/Source/MultiUserServer/Private/Widgets/Clients/Browser/SConcertNetworkBrowser.cpp
+++ b/Engine/Plugins/Developer/Concert/ConcertApp/MultiUserServer/Source/MultiUserServer/Private/Widgets/Clients/Browser/SConcertNetworkBrowser.cpp
@@ -373,11 +373,11 @@ TSharedPtr<SWidget> UE::MultiUserServer::SConcertNetworkBrowser::OnGetContextMen
 		);
 	AddDisplayModeEntry(MenuBuilder, EConcertBrowserItemDisplayMode::OutboundSegementTable,
 		LOCTEXT("DisplayMode.OutboundSegementTable.Title", "Outbound segment table"),
-		LOCTEXT("DisplayMode.SegementTable.Tooltip", "A table displaying the messaging protocol's outbound segments' MessageId, Sent, Acked and Size data in realtime.")
+		LOCTEXT("DisplayMode.OutboundSegementTable.Tooltip", "A table displaying the messaging protocol's outbound segments' MessageId, Sent, Acked and Size data in realtime.")
 		);
 	AddDisplayModeEntry(MenuBuilder, EConcertBrowserItemDisplayMode::InboundSegmentTable,
 		LOCTEXT("DisplayMode.InboundSegementTable.Title", "Inbound Segment table"),
-		LOCTEXT("DisplayMode.SegementTable.Tooltip", "A table displaying the messaging protocol's inbound segments' MessageId, Received and Size data in realtime.")
+		LOCTEXT("DisplayMode.InboundSegementTable.Tooltip", "A table displaying the messaging protocol's inbound segments' MessageId, Received and Size data in realtime.")
 		);
 	return MenuBuilder.MakeWidget();
 }
