Add server item to concert browser

diff --git a/Engine/Plugins/Developer/Concert/ConcertApp/MultiUserServer/Source/MultiUserServer/Private/Widgets/Clients/Browser/ClientBrowserItem.cpp b/Engine/Plugins/Developer/Concert/ConcertApp/MultiUserServer/Source/MultiUserServer/Private/Widgets/Clients/Browser/ClientBrowserItem.cpp
deleted file mode 100644
diff --git a/Engine/Plugins/Developer/Concert/ConcertApp/MultiUserServer/Source/MultiUserServer/Private/Widgets/Clients/Browser/ClientBrowserItem.h b/Engine/Plugins/Developer/Concert/ConcertApp/MultiUserServer/Source/MultiUserServer/Private/Widgets/Clients/Browser/ClientBrowserItem.h
deleted file mode 100644
diff --git a/Engine/Plugins/Developer/Concert/ConcertApp/MultiUserServer/Source/MultiUserServer/Private/Widgets/Clients/Browser/Graph/SClientNetworkGraphs.cpp b/Engine/Plugins/Developer/Concert/ConcertApp/MultiUserServer/Source/MultiUserServer/Private/Widgets/Clients/Browser/Graph/SClientNetworkGraphs.cpp
--- a/Engine/Plugins/Developer/Concert/ConcertApp/MultiUserServer/Source/MultiUserServer/Private/Widgets/Clients/Browser/Graph/SClientNetworkGraphs.cpp
+++ b/Engine/Plugins/Developer/Concert/ConcertApp/MultiUserServer/Source/MultiUserServer/Private/Widgets/Clients/Browser/Graph/SClientNetworkGraphs.cpp
@@ -4,7 +4,7 @@
 
 #include "ConcertServerStyle.h"
 #include "SNetworkGraph.h"
-#include "Widgets/Clients/Browser/Models/IClientTransferStatisticsModel.h"
+#include "Widgets/Clients/Browser/Models/ITransferStatisticsModel.h"
 
 #include "Math/UnitConversion.h"
 #include "Widgets/Layout/SBorder.h"
@@ -15,7 +15,7 @@
 
 namespace UE::MultiUserServer 
 {
-	void SClientNetworkGraphs::Construct(const FArguments& InArgs, TSharedRef<IClientTransferStatisticsModel> InTransferStatsModel)
+	void SClientNetworkGraphs::Construct(const FArguments& InArgs, TSharedRef<ITransferStatisticsModel> InTransferStatsModel)
 	{
 		TransferStatisticsModel = MoveTemp(InTransferStatsModel);
 
diff --git a/Engine/Plugins/Developer/Concert/ConcertApp/MultiUserServer/Source/MultiUserServer/Private/Widgets/Clients/Browser/Graph/SClientNetworkGraphs.h b/Engine/Plugins/Developer/Concert/ConcertApp/MultiUserServer/Source/MultiUserServer/Private/Widgets/Clients/Browser/Graph/SClientNetworkGraphs.h
--- a/Engine/Plugins/Developer/Concert/ConcertApp/MultiUserServer/Source/MultiUserServer/Private/Widgets/Clients/Browser/Graph/SClientNetworkGraphs.h
+++ b/Engine/Plugins/Developer/Concert/ConcertApp/MultiUserServer/Source/MultiUserServer/Private/Widgets/Clients/Browser/Graph/SClientNetworkGraphs.h
@@ -5,11 +5,11 @@
 #include "CoreMinimal.h"
 #include "Widgets/DeclarativeSyntaxSupport.h"
 #include "Widgets/SCompoundWidget.h"
-#include "Widgets/Clients/Browser/Models/IClientTransferStatisticsModel.h"
+#include "Widgets/Clients/Browser/Models/ITransferStatisticsModel.h"
 
 namespace UE::MultiUserServer
 {
-	class IClientTransferStatisticsModel;
+	class ITransferStatisticsModel;
 	class SNetworkGraph;
 
 	/** Displays send and receive graphs */
@@ -24,12 +24,12 @@ namespace UE::MultiUserServer
 			SLATE_ARGUMENT(FTimespan, TimeRange)
 		SLATE_END_ARGS()
 
-		void Construct(const FArguments& InArgs, TSharedRef<IClientTransferStatisticsModel> InTransferStatsModel);
+		void Construct(const FArguments& InArgs, TSharedRef<ITransferStatisticsModel> InTransferStatsModel);
 
 	private:
 
 		/** Tells us the data to display */
-		TSharedPtr<IClientTransferStatisticsModel> TransferStatisticsModel;
+		TSharedPtr<ITransferStatisticsModel> TransferStatisticsModel;
 
 		TMap<EConcertTransferStatistic, TSharedPtr<SNetworkGraph>> Graphs;
 		
diff --git a/Engine/Plugins/Developer/Concert/ConcertApp/MultiUserServer/Source/MultiUserServer/Private/Widgets/Clients/Browser/Graph/SNetworkGraph.cpp b/Engine/Plugins/Developer/Concert/ConcertApp/MultiUserServer/Source/MultiUserServer/Private/Widgets/Clients/Browser/Graph/SNetworkGraph.cpp
--- a/Engine/Plugins/Developer/Concert/ConcertApp/MultiUserServer/Source/MultiUserServer/Private/Widgets/Clients/Browser/Graph/SNetworkGraph.cpp
+++ b/Engine/Plugins/Developer/Concert/ConcertApp/MultiUserServer/Source/MultiUserServer/Private/Widgets/Clients/Browser/Graph/SNetworkGraph.cpp
@@ -98,7 +98,7 @@ namespace UE::MultiUserServer
 		
 		Points = DistributeIntoTimeBuckets(Points, DisplayTimeRangInSeconds, TimeSliceSize, LocalNow, TimeBucketAnchor);
 		CurveData->Points = FlattenCurve(Points);
-		Graph->SetFixedRange(0, MaxDataInDisplayTime);
+		Graph->SetFixedRange(0, GetMaxRange());
 	}
 
 	static TArray<FCurvePoint> DistributeIntoTimeBuckets(TArrayView<FCurvePoint> Curve, float TimeRange, float TimeSliceSize, const FDateTime& Now, const FDateTime& TimeAnchor)
@@ -177,7 +177,7 @@ namespace UE::MultiUserServer
 				.VAlign(VAlign_Center)
 				[
 					SNew(STextBlock)
-					.Text_Lambda([this](){ return GetBytesDisplayText(2 * MaxDataInDisplayTime / 3); })
+					.Text_Lambda([this](){ return GetBytesDisplayText(2 * GetMaxRange() / 3); })
 					.ColorAndOpacity(FConcertServerStyle::Get().GetColor("Concert.Clients.NetworkGraph.HorizontalHelperLine.TextColor"))
 				]
 
@@ -202,7 +202,7 @@ namespace UE::MultiUserServer
 				.VAlign(VAlign_Center)
 				[
 					SNew(STextBlock)
-					.Text_Lambda([this](){ return GetBytesDisplayText(MaxDataInDisplayTime / 3); })
+					.Text_Lambda([this](){ return GetBytesDisplayText(GetMaxRange() / 3); })
 					.ColorAndOpacity(FConcertServerStyle::Get().GetColor("Concert.Clients.NetworkGraph.HorizontalHelperLine.TextColor"))
 				]
 
diff --git a/Engine/Plugins/Developer/Concert/ConcertApp/MultiUserServer/Source/MultiUserServer/Private/Widgets/Clients/Browser/Graph/SNetworkGraph.h b/Engine/Plugins/Developer/Concert/ConcertApp/MultiUserServer/Source/MultiUserServer/Private/Widgets/Clients/Browser/Graph/SNetworkGraph.h
--- a/Engine/Plugins/Developer/Concert/ConcertApp/MultiUserServer/Source/MultiUserServer/Private/Widgets/Clients/Browser/Graph/SNetworkGraph.h
+++ b/Engine/Plugins/Developer/Concert/ConcertApp/MultiUserServer/Source/MultiUserServer/Private/Widgets/Clients/Browser/Graph/SNetworkGraph.h
@@ -86,6 +86,8 @@ namespace UE::MultiUserServer
 
 		/** Updates the graph in an interval of TimeSliceSize in case no network data is received in that time. */
 		void SetTimerForForceRefresh();
+
+		uint64 GetMaxRange() const { return MaxDataInDisplayTime * 1.15; }
 	};
 }
 
diff --git a/Engine/Plugins/Developer/Concert/ConcertApp/MultiUserServer/Source/MultiUserServer/Private/Widgets/Clients/Browser/IConcertBrowserItem.cpp b/Engine/Plugins/Developer/Concert/ConcertApp/MultiUserServer/Source/MultiUserServer/Private/Widgets/Clients/Browser/IConcertBrowserItem.cpp
deleted file mode 100644
diff --git a/Engine/Plugins/Developer/Concert/ConcertApp/MultiUserServer/Source/MultiUserServer/Private/Widgets/Clients/Browser/IConcertBrowserItem.h b/Engine/Plugins/Developer/Concert/ConcertApp/MultiUserServer/Source/MultiUserServer/Private/Widgets/Clients/Browser/IConcertBrowserItem.h
deleted file mode 100644
diff --git a/Engine/Plugins/Developer/Concert/ConcertApp/MultiUserServer/Source/MultiUserServer/Private/Widgets/Clients/Browser/Item/ClientBrowserItem.cpp b/Engine/Plugins/Developer/Concert/ConcertApp/MultiUserServer/Source/MultiUserServer/Private/Widgets/Clients/Browser/Item/ClientBrowserItem.cpp
new file mode 100644
--- /dev/null
+++ b/Engine/Plugins/Developer/Concert/ConcertApp/MultiUserServer/Source/MultiUserServer/Private/Widgets/Clients/Browser/Item/ClientBrowserItem.cpp
@@ -0,0 +1,68 @@
+// Copyright Epic Games, Inc. All Rights Reserved.
+
+#include "ClientBrowserItem.h"
+
+#include "Widgets/Clients/Browser/Models/IClientNetworkStatisticsModel.h"
+#include "Widgets/Clients/Browser/Models/Transfer/ClientTransferStatisticsModel.h"
+#include "Widgets/Clients/Util/EndpointToUserNameCache.h"
+
+#define LOCTEXT_NAMESPACE "UnrealMultiUserUI.FClientBrowserItem"
+
+namespace UE::MultiUserServer
+{
+	FClientBrowserItem::FClientBrowserItem(TSharedRef<IClientNetworkStatisticsModel> NetworkStatisticsModel, TSharedRef<FEndpointToUserNameCache> UserLookup, const FMessageAddress& ClientAddress, const FGuid& MessageNodeId)
+		: NetworkStatisticsModel(MoveTemp(NetworkStatisticsModel))
+		, UserLookup(MoveTemp(UserLookup))
+		, TransferStatisticsModel(MakeShared<FClientTransferStatisticsModel>(ClientAddress))
+		, ClientAddress(ClientAddress)
+		, MessageNodeId(MessageNodeId)
+	{}
+	
+	FString FClientBrowserItem::GetDisplayName() const
+	{
+		if (const TOptional<FConcertClientInfo> ClientInfo = UserLookup->GetClientInfoFromNodeId(MessageNodeId))
+		{
+			return ClientInfo->DisplayName;
+		}
+		
+		const FString NodeIdAsString = MessageNodeId.ToString(EGuidFormats::DigitsWithHyphens);
+		int32 Index = INDEX_NONE;
+		const bool bFound = NodeIdAsString.FindChar('-', Index);
+		// Avoid making the name too long by only showing the first few digits of the node ID
+		return FString::Printf(TEXT("Admin (%s)"), bFound ? *NodeIdAsString.Left(Index) : *NodeIdAsString); 
+	}
+
+	FText FClientBrowserItem::GetToolTip() const
+	{
+		if (const TOptional<FConcertClientInfo> ClientInfo = UserLookup->GetClientInfoFromNodeId(MessageNodeId))
+		{
+			return FText::Format(
+				LOCTEXT("Name.Available.TooltipFmt", "NodeID: {0}\nAddress ID: {1}"),
+				FText::FromString(MessageNodeId.ToString(EGuidFormats::DigitsWithHyphens)),
+				FText::FromString(ClientAddress.ToString())
+				);
+		}
+		return FText::Format(
+			LOCTEXT("Name.NotAvailable.TooltipFmt", "This client's display information becomes available after joining a session.\nNodeID: {0}\nAddress ID: {1}"),
+			FText::FromString(MessageNodeId.ToString(EGuidFormats::DigitsWithHyphens)),
+			FText::FromString(ClientAddress.ToString())
+			);
+	}
+
+	TSharedRef<ITransferStatisticsModel> FClientBrowserItem::GetTransferStatistics() const
+	{
+		return TransferStatisticsModel;
+	}
+
+	TOptional<FMessageTransportStatistics> FClientBrowserItem::GetLatestNetworkStatistics() const
+	{
+		return NetworkStatisticsModel->GetLatestNetworkStatistics(ClientAddress);
+	}
+
+	bool FClientBrowserItem::IsOnline() const
+	{
+		return NetworkStatisticsModel->IsOnline(ClientAddress);
+	}
+}
+
+#undef LOCTEXT_NAMESPACE
\ No newline at end of file
diff --git a/Engine/Plugins/Developer/Concert/ConcertApp/MultiUserServer/Source/MultiUserServer/Private/Widgets/Clients/Browser/Item/ClientBrowserItem.h b/Engine/Plugins/Developer/Concert/ConcertApp/MultiUserServer/Source/MultiUserServer/Private/Widgets/Clients/Browser/Item/ClientBrowserItem.h
new file mode 100644
--- /dev/null
+++ b/Engine/Plugins/Developer/Concert/ConcertApp/MultiUserServer/Source/MultiUserServer/Private/Widgets/Clients/Browser/Item/ClientBrowserItem.h
@@ -0,0 +1,63 @@
+// Copyright Epic Games, Inc. All Rights Reserved.
+
+#pragma once
+
+#include "CoreMinimal.h"
+#include "ConcertMessageData.h"
+#include "IConcertBrowserItem.h"
+#include "IMessageContext.h"
+
+class FEndpointToUserNameCache;
+
+namespace UE::MultiUserServer
+{
+	class ITransferStatisticsModel;
+	class FClientTransferStatisticsModel;
+
+	DECLARE_DELEGATE_RetVal(TOptional<FConcertClientInfo>, FGetClientInfo);
+	
+	class FClientBrowserItem : public FConcertBrowserItemCommonImpl
+	{
+	public:
+
+		FClientBrowserItem(TSharedRef<IClientNetworkStatisticsModel> NetworkStatisticsModel, TSharedRef<FEndpointToUserNameCache> UserLookup, const FMessageAddress& ClientAddress, const FGuid& MessageNodeId);
+		
+		void OnJoinSession(const FGuid& SessionId) { CurrentSession = SessionId; }
+		void OnLeaveSession() { CurrentSession.Reset(); }
+		const TOptional<FGuid>& GetCurrentSession() const { return CurrentSession; }
+		const FGuid& GetMessageNodeId() const { return MessageNodeId; }
+
+		void OnClientDisconnected() { bIsDisconnected = true; }
+		void OnClientReconnected() { bIsDisconnected = false; }
+		/** Whether this client is no longer connected with the server */
+		bool IsDisconnected() const { return bIsDisconnected; }
+		
+		virtual FString GetDisplayName() const override;
+		virtual FText GetToolTip() const override;
+		virtual TSharedRef<ITransferStatisticsModel> GetTransferStatistics() const override;
+		virtual TOptional<FMessageTransportStatistics> GetLatestNetworkStatistics() const override;
+		virtual bool IsOnline() const override;
+
+	private:
+
+		TSharedRef<IClientNetworkStatisticsModel> NetworkStatisticsModel;
+		
+		/** Used to gets the client's name once it becomes available */
+		TSharedRef<FEndpointToUserNameCache> UserLookup;
+
+		/** Keeps track of the clients transfer statistics */
+		TSharedRef<FClientTransferStatisticsModel> TransferStatisticsModel;
+
+		/** Address of this client as used by the UDP messaging sytem */
+		FMessageAddress ClientAddress;
+		
+		/** ID of this client in the UDP messaging system */
+		FGuid MessageNodeId;
+		
+		/** The session this client currently is in */
+		TOptional<FGuid> CurrentSession;
+
+		/** Whether this client is no longer connected with the server */
+		bool bIsDisconnected = false;
+	};
+}
diff --git a/Engine/Plugins/Developer/Concert/ConcertApp/MultiUserServer/Source/MultiUserServer/Private/Widgets/Clients/Browser/Item/IConcertBrowserItem.cpp b/Engine/Plugins/Developer/Concert/ConcertApp/MultiUserServer/Source/MultiUserServer/Private/Widgets/Clients/Browser/Item/IConcertBrowserItem.cpp
new file mode 100644
--- /dev/null
+++ b/Engine/Plugins/Developer/Concert/ConcertApp/MultiUserServer/Source/MultiUserServer/Private/Widgets/Clients/Browser/Item/IConcertBrowserItem.cpp
@@ -0,0 +1,37 @@
+// Copyright Epic Games, Inc. All Rights Reserved.
+
+#include "IConcertBrowserItem.h"
+
+#include "Widgets/Clients/Browser/Models/IClientNetworkStatisticsModel.h"
+
+#include "IMessageContext.h"
+#include "INetworkMessagingExtension.h"
+
+namespace UE::MultiUserServer
+{
+	void FConcertBrowserItemCommonImpl::SetDisplayMode(EConcertBrowserItemDisplayMode Value)
+	{
+		ItemDisplayMode = Value;
+	}
+	
+	EConcertBrowserItemDisplayMode FConcertBrowserItemCommonImpl::GetDisplayMode() const
+	{
+		return ItemDisplayMode;
+	}
+
+	void FConcertBrowserItemCommonImpl::AppendSearchTerms(TArray<FString>& SearchTerms) const
+	{
+		SearchTerms.Add(GetDisplayName());
+		const TOptional<FMessageTransportStatistics> Stats = GetLatestNetworkStatistics();
+		SearchTerms.Add(NetworkStatistics::FormatIPv4AsString(Stats));
+		
+		if (Stats)
+		{
+			SearchTerms.Add(NetworkStatistics::FormatTotalBytesSent(*Stats));
+			SearchTerms.Add(NetworkStatistics::FormatTotalBytesReceived(*Stats));
+			SearchTerms.Add(NetworkStatistics::FormatAverageRTT(*Stats));
+			SearchTerms.Add(NetworkStatistics::FormatBytesInflight(*Stats));
+			SearchTerms.Add(NetworkStatistics::FormatTotalBytesLost(*Stats));
+		}
+	}
+}
diff --git a/Engine/Plugins/Developer/Concert/ConcertApp/MultiUserServer/Source/MultiUserServer/Private/Widgets/Clients/Browser/Item/IConcertBrowserItem.h b/Engine/Plugins/Developer/Concert/ConcertApp/MultiUserServer/Source/MultiUserServer/Private/Widgets/Clients/Browser/Item/IConcertBrowserItem.h
new file mode 100644
--- /dev/null
+++ b/Engine/Plugins/Developer/Concert/ConcertApp/MultiUserServer/Source/MultiUserServer/Private/Widgets/Clients/Browser/Item/IConcertBrowserItem.h
@@ -0,0 +1,65 @@
+// Copyright Epic Games, Inc. All Rights Reserved.
+
+#pragma once
+
+#include "CoreMinimal.h"
+
+struct FMessageTransportStatistics;
+struct FMessageAddress;
+
+namespace UE::MultiUserServer
+{
+	class IClientNetworkStatisticsModel;
+	class ITransferStatisticsModel;
+
+	enum class EConcertBrowserItemDisplayMode : uint8
+	{
+		/** Displays the sent and read packets */
+		NetworkGraph = 0,
+		/** Displays a table showing MessageId, Sent, Acked, and Size updated in realtime for outbound UDP segments. */
+		OutboundSegementTable = 1,
+		/** Displays a table showing MessageId, Received, and Size updated in realtime for inbound UDP segments. */
+		InboundSegmentTable = 2
+	};
+
+	/** Implemented by items that can be displayed in the Concert Network Browser */
+	class IConcertBrowserItem
+	{
+	public:
+
+		/** The title of the thumbnail */
+		virtual FString GetDisplayName() const = 0;
+		/** What to display when the thumbnail is hovered */
+		virtual FText GetToolTip() const = 0;
+		
+		virtual void SetDisplayMode(EConcertBrowserItemDisplayMode Value) = 0;
+		virtual EConcertBrowserItemDisplayMode GetDisplayMode() const = 0;
+
+		/** The network statistics for this item */
+		virtual TSharedRef<ITransferStatisticsModel> GetTransferStatistics() const = 0;
+
+		/** Adds terms this item can be searched by */
+		virtual void AppendSearchTerms(TArray<FString>& SearchTerms) const = 0;
+
+		/** Gets the latest network statistics for this item if they are available; most likely unavailable when IsOnline returns false. */
+		virtual TOptional<FMessageTransportStatistics> GetLatestNetworkStatistics() const = 0;
+		/** Whether this client is currently reachable */
+		virtual bool IsOnline() const = 0;
+
+		virtual ~IConcertBrowserItem() = default;
+	};
+
+	/** Util for implementing shared logic. */
+	class FConcertBrowserItemCommonImpl : public IConcertBrowserItem
+	{
+	public:
+		
+		virtual void SetDisplayMode(EConcertBrowserItemDisplayMode Value) override;
+		virtual EConcertBrowserItemDisplayMode GetDisplayMode() const override;
+		virtual void AppendSearchTerms(TArray<FString>& SearchTerms) const override;
+
+	private:
+
+		EConcertBrowserItemDisplayMode ItemDisplayMode = EConcertBrowserItemDisplayMode::NetworkGraph;
+	};
+}
diff --git a/Engine/Plugins/Developer/Concert/ConcertApp/MultiUserServer/Source/MultiUserServer/Private/Widgets/Clients/Browser/Item/ServerBrowserItem.cpp b/Engine/Plugins/Developer/Concert/ConcertApp/MultiUserServer/Source/MultiUserServer/Private/Widgets/Clients/Browser/Item/ServerBrowserItem.cpp
new file mode 100644
--- /dev/null
+++ b/Engine/Plugins/Developer/Concert/ConcertApp/MultiUserServer/Source/MultiUserServer/Private/Widgets/Clients/Browser/Item/ServerBrowserItem.cpp
@@ -0,0 +1,98 @@
+// Copyright Epic Games, Inc. All Rights Reserved.
+
+#include "ServerBrowserItem.h"
+
+#include "IConcertServer.h"
+#include "Widgets/Clients/Browser/Models/IClientNetworkStatisticsModel.h"
+#include "Widgets/Clients/Browser/Models/Transfer/ServerTransferStatisticsModel.h"
+
+#define LOCTEXT_NAMESPACE "UnrealMultiUserUI.FServerBrowserItem"
+
+namespace UE::MultiUserServer
+{
+	FServerBrowserItem::FServerBrowserItem(TSharedRef<IConcertServer> Server, TSharedRef<IClientNetworkStatisticsModel> NetworkStatisticsModel)
+		: Server(MoveTemp(Server))
+		, NetworkStatisticsModel(MoveTemp(NetworkStatisticsModel))
+		, ServerTransferStats(MakeShared<FServerTransferStatisticsModel>())
+	{}
+
+	FString FServerBrowserItem::GetDisplayName() const
+	{
+		return TEXT("Server");
+	}
+
+	FText FServerBrowserItem::GetToolTip() const
+	{
+		return LOCTEXT("ServerTooltip", "The server sums up all connected client statistics.");
+	}
+
+	TSharedRef<ITransferStatisticsModel> FServerBrowserItem::GetTransferStatistics() const
+	{
+		return ServerTransferStats;
+	}
+
+	TOptional<FMessageTransportStatistics> FServerBrowserItem::GetLatestNetworkStatistics() const
+	{
+		TArray<FMessageTransportStatistics> AllStats;
+		for (const TSharedPtr<IConcertServerSession>& Session : Server->GetLiveSessions())
+		{
+			for (const FConcertSessionClientInfo& ClientInfo : Session->GetSessionClients())
+			{
+				const FMessageAddress ClientAddress = Session->GetClientAddress(ClientInfo.ClientEndpointId);
+				const TOptional<FMessageTransportStatistics> Stats = NetworkStatisticsModel->GetLatestNetworkStatistics(ClientAddress);
+				if (ensureMsgf(Stats, TEXT("Server stats should always be available because the connection is server controlled.")))
+				{
+					AllStats.Add(*Stats);
+				}
+			}
+			
+		}
+
+		for (const FConcertEndpointContext& RemoteClientAdminEndpoints : Server->GetRemoteAdminEndpoints())
+		{
+			const TOptional<FMessageTransportStatistics> Stats = NetworkStatisticsModel->GetLatestNetworkStatistics(Server->GetRemoteAddress(RemoteClientAdminEndpoints.EndpointId));
+			if (ensureMsgf(Stats, TEXT("Server stats should always be available because the connection is server controlled.")))
+			{
+				AllStats.Add(*Stats);
+			}
+		}
+		
+		if (AllStats.Num() == 0)
+		{
+			return {};
+		}
+
+		FMessageTransportStatistics Result = SumUpStats(AllStats);
+		Result.IPv4AsString = TEXT("localhost"); // Should probably look up the IP address values ...
+		return Result;
+	}
+
+	FMessageTransportStatistics FServerBrowserItem::SumUpStats(const TArray<FMessageTransportStatistics>& StatArray) const
+	{
+		FMessageTransportStatistics Result;
+
+		for (const FMessageTransportStatistics& Stats : StatArray)
+		{
+			Result.TotalBytesSent += Stats.TotalBytesSent;
+			Result.TotalBytesLost += Stats.TotalBytesLost;
+			Result.TotalBytesReceived += Stats.TotalBytesReceived;
+			Result.BytesInflight += Stats.BytesInflight;
+			Result.PacketsSent += Stats.PacketsSent;
+			Result.PacketsLost += Stats.PacketsLost;
+			Result.PacketsAcked += Stats.PacketsAcked;
+			Result.PacketsReceived += Stats.PacketsReceived;
+			Result.PacketsInFlight += Stats.PacketsInFlight;
+			Result.AverageRTT += Stats.AverageRTT;
+		}
+
+		if (ensure(StatArray.Num() > 0))
+		{
+			Result.AverageRTT /= StatArray.Num();
+		}
+		Result.WindowSize = 0;
+
+		return Result;
+	}
+}
+
+#undef LOCTEXT_NAMESPACE
\ No newline at end of file
diff --git a/Engine/Plugins/Developer/Concert/ConcertApp/MultiUserServer/Source/MultiUserServer/Private/Widgets/Clients/Browser/Item/ServerBrowserItem.h b/Engine/Plugins/Developer/Concert/ConcertApp/MultiUserServer/Source/MultiUserServer/Private/Widgets/Clients/Browser/Item/ServerBrowserItem.h
new file mode 100644
--- /dev/null
+++ b/Engine/Plugins/Developer/Concert/ConcertApp/MultiUserServer/Source/MultiUserServer/Private/Widgets/Clients/Browser/Item/ServerBrowserItem.h
@@ -0,0 +1,36 @@
+// Copyright Epic Games, Inc. All Rights Reserved.
+
+#pragma once
+
+#include "CoreMinimal.h"
+#include "IConcertBrowserItem.h"
+#include "IMessageContext.h"
+
+class IConcertServer;
+
+namespace UE::MultiUserServer
+{
+	class FServerTransferStatisticsModel;
+
+	class FServerBrowserItem : public FConcertBrowserItemCommonImpl
+	{
+	public:
+
+		FServerBrowserItem(TSharedRef<IConcertServer> Server, TSharedRef<IClientNetworkStatisticsModel> NetworkStatisticsModel);
+	
+		virtual FString GetDisplayName() const override;
+		virtual FText GetToolTip() const override;
+		virtual TSharedRef<ITransferStatisticsModel> GetTransferStatistics() const override;
+		virtual TOptional<FMessageTransportStatistics> GetLatestNetworkStatistics() const override;
+		virtual bool IsOnline() const override { return true;}
+
+	private:
+
+		TSharedRef<IConcertServer> Server;
+		TSharedRef<IClientNetworkStatisticsModel> NetworkStatisticsModel;
+		TSharedRef<FServerTransferStatisticsModel> ServerTransferStats;
+
+		FMessageTransportStatistics SumUpStats(const TArray<FMessageTransportStatistics>& Stats) const;
+	};
+}
+
diff --git a/Engine/Plugins/Developer/Concert/ConcertApp/MultiUserServer/Source/MultiUserServer/Private/Widgets/Clients/Browser/Models/ClientBrowserModel.cpp b/Engine/Plugins/Developer/Concert/ConcertApp/MultiUserServer/Source/MultiUserServer/Private/Widgets/Clients/Browser/Models/ClientBrowserModel.cpp
--- a/Engine/Plugins/Developer/Concert/ConcertApp/MultiUserServer/Source/MultiUserServer/Private/Widgets/Clients/Browser/Models/ClientBrowserModel.cpp
+++ b/Engine/Plugins/Developer/Concert/ConcertApp/MultiUserServer/Source/MultiUserServer/Private/Widgets/Clients/Browser/Models/ClientBrowserModel.cpp
@@ -4,7 +4,7 @@
 
 #include "ConcertServerEvents.h"
 #include "IConcertServer.h"
-#include "Widgets/Clients/Browser/ClientBrowserItem.h"
+#include "Widgets/Clients/Browser/Item/ClientBrowserItem.h"
 #include "Widgets/Clients/Util/EndpointToUserNameCache.h"
 
 #include "Algo/AnyOf.h"
diff --git a/Engine/Plugins/Developer/Concert/ConcertApp/MultiUserServer/Source/MultiUserServer/Private/Widgets/Clients/Browser/Models/ClientTransferStatisticsModel.cpp b/Engine/Plugins/Developer/Concert/ConcertApp/MultiUserServer/Source/MultiUserServer/Private/Widgets/Clients/Browser/Models/ClientTransferStatisticsModel.cpp
deleted file mode 100644
diff --git a/Engine/Plugins/Developer/Concert/ConcertApp/MultiUserServer/Source/MultiUserServer/Private/Widgets/Clients/Browser/Models/ClientTransferStatisticsModel.h b/Engine/Plugins/Developer/Concert/ConcertApp/MultiUserServer/Source/MultiUserServer/Private/Widgets/Clients/Browser/Models/ClientTransferStatisticsModel.h
deleted file mode 100644
diff --git a/Engine/Plugins/Developer/Concert/ConcertApp/MultiUserServer/Source/MultiUserServer/Private/Widgets/Clients/Browser/Models/IClientTransferStatisticsModel.h b/Engine/Plugins/Developer/Concert/ConcertApp/MultiUserServer/Source/MultiUserServer/Private/Widgets/Clients/Browser/Models/IClientTransferStatisticsModel.h
deleted file mode 100644
diff --git a/Engine/Plugins/Developer/Concert/ConcertApp/MultiUserServer/Source/MultiUserServer/Private/Widgets/Clients/Browser/Models/ITransferStatisticsModel.h b/Engine/Plugins/Developer/Concert/ConcertApp/MultiUserServer/Source/MultiUserServer/Private/Widgets/Clients/Browser/Models/ITransferStatisticsModel.h
new file mode 100644
--- /dev/null
+++ b/Engine/Plugins/Developer/Concert/ConcertApp/MultiUserServer/Source/MultiUserServer/Private/Widgets/Clients/Browser/Models/ITransferStatisticsModel.h
@@ -0,0 +1,61 @@
+// Copyright Epic Games, Inc. All Rights Reserved.
+
+#pragma once
+
+#include "CoreMinimal.h"
+#include "INetworkMessagingExtension.h"
+
+struct FMessageAddress;
+
+namespace UE::MultiUserServer
+{
+	enum class EConcertTransferStatistic
+	{
+		SentToClient,
+		ReceivedFromClient,
+
+		Count
+	};
+
+	DECLARE_MULTICAST_DELEGATE(FOnTransferTimelineUpdated);
+	DECLARE_MULTICAST_DELEGATE(FOnTransferGroupsUpdated);
+	
+	struct FConcertTransferSamplePoint
+	{
+		/** The time these stats were updated */
+		FDateTime LocalTime;
+
+		/** The number of bytes transferred, e.g. sent, received, etc. */
+		uint64 BytesTransferred;
+
+		FConcertTransferSamplePoint() = default;
+		explicit FConcertTransferSamplePoint(const FDateTime& Time, uint64 BytesTransferred = 0)
+			: LocalTime(Time)
+			, BytesTransferred(BytesTransferred)
+		{}
+	};
+	
+	/** Keeps track of a single client's transfer statistics */
+	class ITransferStatisticsModel
+	{
+	public:
+
+		/** @return Contains the transfer statistics over time. Used for graphs. */
+		virtual const TArray<FConcertTransferSamplePoint>& GetTransferStatTimeline(EConcertTransferStatistic StatisticType) const = 0;
+		/** @return Sorted descending by message ID. Every message ID is unique. Used for list views. */
+		virtual const TArray<TSharedPtr<FOutboundTransferStatistics>>& GetOutboundTransferStatsGroupedById() const = 0;
+		/** @return Sorted descending by message ID. Every message ID is unique. Used for list views. */
+		virtual const TArray<TSharedPtr<FInboundTransferStatistics>>& GetInboundTransferStatsGroupedById() const = 0;
+
+		/** Called when GetTransferStatTimeline changes */
+		virtual FOnTransferTimelineUpdated& OnTransferTimelineUpdated(EConcertTransferStatistic StatisticType) = 0;
+		/** Called when GetTransferStatsGroupedById changes */
+		virtual FOnTransferGroupsUpdated& OnOutboundTransferGroupsUpdated() = 0;
+		/** Called when GetTransferStatsGroupedById changes */
+		virtual FOnTransferGroupsUpdated& OnInboundTransferGroupsUpdated() = 0;
+
+		virtual ~ITransferStatisticsModel() = default;
+	};
+}
+
+
diff --git a/Engine/Plugins/Developer/Concert/ConcertApp/MultiUserServer/Source/MultiUserServer/Private/Widgets/Clients/Browser/Models/Transfer/ClientTransferStatisticsModel.cpp b/Engine/Plugins/Developer/Concert/ConcertApp/MultiUserServer/Source/MultiUserServer/Private/Widgets/Clients/Browser/Models/Transfer/ClientTransferStatisticsModel.cpp
new file mode 100644
--- /dev/null
+++ b/Engine/Plugins/Developer/Concert/ConcertApp/MultiUserServer/Source/MultiUserServer/Private/Widgets/Clients/Browser/Models/Transfer/ClientTransferStatisticsModel.cpp
@@ -0,0 +1,66 @@
+// Copyright Epic Games, Inc. All Rights Reserved.
+
+#include "ClientTransferStatisticsModel.h"
+
+#include "INetworkMessagingExtension.h"
+#include "Features/IModularFeatures.h"
+
+namespace UE::MultiUserServer
+{
+	namespace Private::ClientTransferStatisticsModel
+	{
+		static INetworkMessagingExtension* GetMessagingStatistics()
+		{
+			IModularFeatures& ModularFeatures = IModularFeatures::Get();
+			if (IsInGameThread())
+			{
+				if (ModularFeatures.IsModularFeatureAvailable(INetworkMessagingExtension::ModularFeatureName))
+				{
+					return &ModularFeatures.GetModularFeature<INetworkMessagingExtension>(INetworkMessagingExtension::ModularFeatureName);
+				}
+			}
+			else
+			{
+				ModularFeatures.LockModularFeatureList();
+				ON_SCOPE_EXIT
+				{
+					ModularFeatures.UnlockModularFeatureList();
+				};
+			
+				if (ModularFeatures.IsModularFeatureAvailable(INetworkMessagingExtension::ModularFeatureName))
+				{
+					return &ModularFeatures.GetModularFeature<INetworkMessagingExtension>(INetworkMessagingExtension::ModularFeatureName);
+				}
+			}
+		
+			ensureMsgf(false, TEXT("Feature %s is unavailable"), *INetworkMessagingExtension::ModularFeatureName.ToString());
+			return nullptr;
+		}
+	}
+	
+	FClientTransferStatisticsModel::FClientTransferStatisticsModel(const FMessageAddress& ClientAddress)
+		: FTransferStatisticsModelBase()
+		, ClientAddress(ClientAddress)
+	{
+	}
+
+	bool FClientTransferStatisticsModel::ShouldIncludeOutboundStat(const FOutboundTransferStatistics& Item) const
+	{
+		if (const INetworkMessagingExtension* Statistics = Private::ClientTransferStatisticsModel::GetMessagingStatistics())
+		{
+			const FGuid ClientNodeId = Statistics->GetNodeIdFromAddress(ClientAddress);
+			return ClientNodeId == Item.DestinationId;
+		}
+		return false;
+	}
+
+	bool FClientTransferStatisticsModel::ShouldIncludeInboundStat(const FInboundTransferStatistics& Item) const
+	{
+		if (const INetworkMessagingExtension* Statistics = Private::ClientTransferStatisticsModel::GetMessagingStatistics())
+		{
+			const FGuid ClientNodeId = Statistics->GetNodeIdFromAddress(ClientAddress);
+			return ClientNodeId == Item.OriginId;
+		}
+		return false; 
+	}
+}
\ No newline at end of file
diff --git a/Engine/Plugins/Developer/Concert/ConcertApp/MultiUserServer/Source/MultiUserServer/Private/Widgets/Clients/Browser/Models/Transfer/ClientTransferStatisticsModel.h b/Engine/Plugins/Developer/Concert/ConcertApp/MultiUserServer/Source/MultiUserServer/Private/Widgets/Clients/Browser/Models/Transfer/ClientTransferStatisticsModel.h
new file mode 100644
--- /dev/null
+++ b/Engine/Plugins/Developer/Concert/ConcertApp/MultiUserServer/Source/MultiUserServer/Private/Widgets/Clients/Browser/Models/Transfer/ClientTransferStatisticsModel.h
@@ -0,0 +1,26 @@
+// Copyright Epic Games, Inc. All Rights Reserved.
+
+#pragma once
+
+#include "CoreMinimal.h"
+#include "IMessageContext.h"
+#include "TransferStatisticsModel.h"
+#include "Widgets/Clients/Browser/Models/ITransferStatisticsModel.h"
+
+namespace UE::MultiUserServer
+{
+	class FClientTransferStatisticsModel : public FTransferStatisticsModelBase
+	{
+	public:
+		
+		FClientTransferStatisticsModel(const FMessageAddress& ClientAddress);
+
+	protected:
+
+		const FMessageAddress ClientAddress;
+		
+		virtual bool ShouldIncludeOutboundStat(const FOutboundTransferStatistics& Item) const override;
+		virtual bool ShouldIncludeInboundStat(const FInboundTransferStatistics& Item) const override;
+	};
+}
+
diff --git a/Engine/Plugins/Developer/Concert/ConcertApp/MultiUserServer/Source/MultiUserServer/Private/Widgets/Clients/Browser/Models/Transfer/ServerTransferStatisticsModel.h b/Engine/Plugins/Developer/Concert/ConcertApp/MultiUserServer/Source/MultiUserServer/Private/Widgets/Clients/Browser/Models/Transfer/ServerTransferStatisticsModel.h
new file mode 100644
--- /dev/null
+++ b/Engine/Plugins/Developer/Concert/ConcertApp/MultiUserServer/Source/MultiUserServer/Private/Widgets/Clients/Browser/Models/Transfer/ServerTransferStatisticsModel.h
@@ -0,0 +1,19 @@
+// Copyright Epic Games, Inc. All Rights Reserved.
+
+#pragma once
+
+#include "CoreMinimal.h"
+#include "TransferStatisticsModel.h"
+#include "Widgets/Clients/Browser/Models/ITransferStatisticsModel.h"
+
+namespace UE::MultiUserServer
+{
+	class FServerTransferStatisticsModel : public FTransferStatisticsModelBase
+	{
+	protected:
+		
+		virtual bool ShouldIncludeOutboundStat(const FOutboundTransferStatistics& Item) const override { return true; }
+		virtual bool ShouldIncludeInboundStat(const FInboundTransferStatistics& Item) const override { return true; }
+	};
+}
+
diff --git a/Engine/Plugins/Developer/Concert/ConcertApp/MultiUserServer/Source/MultiUserServer/Private/Widgets/Clients/Browser/Models/Transfer/TransferStatisticsModel.cpp b/Engine/Plugins/Developer/Concert/ConcertApp/MultiUserServer/Source/MultiUserServer/Private/Widgets/Clients/Browser/Models/Transfer/TransferStatisticsModel.cpp
new file mode 100644
--- /dev/null
+++ b/Engine/Plugins/Developer/Concert/ConcertApp/MultiUserServer/Source/MultiUserServer/Private/Widgets/Clients/Browser/Models/Transfer/TransferStatisticsModel.cpp
@@ -0,0 +1,97 @@
+// Copyright Epic Games, Inc. All Rights Reserved.
+
+#include "TransferStatisticsModel.h"
+
+#include "INetworkMessagingExtension.h"
+#include "Containers/Ticker.h"
+#include "Features/IModularFeatures.h"
+
+namespace UE::MultiUserServer
+{
+	namespace Private::TransferStatisticsModel
+	{
+		static INetworkMessagingExtension* GetMessagingStatistics()
+		{
+			IModularFeatures& ModularFeatures = IModularFeatures::Get();
+			if (IsInGameThread())
+			{
+				if (ModularFeatures.IsModularFeatureAvailable(INetworkMessagingExtension::ModularFeatureName))
+				{
+					return &ModularFeatures.GetModularFeature<INetworkMessagingExtension>(INetworkMessagingExtension::ModularFeatureName);
+				}
+			}
+			else
+			{
+				ModularFeatures.LockModularFeatureList();
+				ON_SCOPE_EXIT
+				{
+					ModularFeatures.UnlockModularFeatureList();
+				};
+			
+				if (ModularFeatures.IsModularFeatureAvailable(INetworkMessagingExtension::ModularFeatureName))
+				{
+					return &ModularFeatures.GetModularFeature<INetworkMessagingExtension>(INetworkMessagingExtension::ModularFeatureName);
+				}
+			}
+		
+			ensureMsgf(false, TEXT("Feature %s is unavailable"), *INetworkMessagingExtension::ModularFeatureName.ToString());
+			return nullptr;
+		}
+	}
+
+	FTransferStatisticsModelBase::FTransferStatisticsModelBase()
+		: OutboundStatTracker([this](const FOutboundTransferStatistics& Stats) { return ShouldIncludeOutboundStat(Stats); }, [](const FOutboundTransferStatistics& Stats){ return Stats.BytesSent; })
+		, InboundStatTracker([this](const FInboundTransferStatistics& Stats) { return ShouldIncludeInboundStat(Stats); }, [](const FInboundTransferStatistics& Stats){ return Stats.BytesReceived; })
+	{
+		if (INetworkMessagingExtension* Statistics = Private::TransferStatisticsModel::GetMessagingStatistics())
+		{
+			Statistics->OnOutboundTransferUpdatedFromThread().AddRaw(&OutboundStatTracker, &TClientTransferStatTracker<FOutboundTransferStatistics>::OnTransferUpdatedFromThread);
+			Statistics->OnInboundTransferUpdatedFromThread().AddRaw(&InboundStatTracker, &TClientTransferStatTracker<FInboundTransferStatistics>::OnTransferUpdatedFromThread);
+			TickHandle = FTSTicker::GetCoreTicker().AddTicker(FTickerDelegate::CreateRaw(this, &FTransferStatisticsModelBase::Tick));
+		}
+	}
+
+	FTransferStatisticsModelBase::~FTransferStatisticsModelBase()
+	{
+		if (INetworkMessagingExtension* Statistics = Private::TransferStatisticsModel::GetMessagingStatistics())
+		{
+			Statistics->OnOutboundTransferUpdatedFromThread().RemoveAll(&OutboundStatTracker);
+			Statistics->OnInboundTransferUpdatedFromThread().RemoveAll(&InboundStatTracker);
+		}
+		
+		FTSTicker::GetCoreTicker().RemoveTicker(TickHandle);
+	}
+
+	const TArray<FConcertTransferSamplePoint>& FTransferStatisticsModelBase::GetTransferStatTimeline(EConcertTransferStatistic StatisticType) const
+	{
+		switch (StatisticType)
+		{
+		case EConcertTransferStatistic::SentToClient: return OutboundStatTracker.GetTransferStatisticsTimeline();
+		case EConcertTransferStatistic::ReceivedFromClient: return InboundStatTracker.GetTransferStatisticsTimeline();;
+		case EConcertTransferStatistic::Count:
+		default:
+			checkNoEntry();
+			return OutboundStatTracker.GetTransferStatisticsTimeline();
+		}
+	}
+
+	FOnTransferTimelineUpdated& FTransferStatisticsModelBase::OnTransferTimelineUpdated(EConcertTransferStatistic StatisticType)
+	{
+		switch (StatisticType)
+		{
+		case EConcertTransferStatistic::SentToClient: return OutboundStatTracker.GetOnTimelineUpdatedDelegates();
+		case EConcertTransferStatistic::ReceivedFromClient: return InboundStatTracker.GetOnTimelineUpdatedDelegates();;
+		case EConcertTransferStatistic::Count:
+		default:
+			checkNoEntry();
+			return OutboundStatTracker.GetOnTimelineUpdatedDelegates();
+		}
+	}
+	
+	bool FTransferStatisticsModelBase::Tick(float DeltaTime)
+	{
+		OutboundStatTracker.Tick();
+		InboundStatTracker.Tick();
+		return true;
+	}
+}
\ No newline at end of file
diff --git a/Engine/Plugins/Developer/Concert/ConcertApp/MultiUserServer/Source/MultiUserServer/Private/Widgets/Clients/Browser/Models/Transfer/TransferStatisticsModel.h b/Engine/Plugins/Developer/Concert/ConcertApp/MultiUserServer/Source/MultiUserServer/Private/Widgets/Clients/Browser/Models/Transfer/TransferStatisticsModel.h
new file mode 100644
--- /dev/null
+++ b/Engine/Plugins/Developer/Concert/ConcertApp/MultiUserServer/Source/MultiUserServer/Private/Widgets/Clients/Browser/Models/Transfer/TransferStatisticsModel.h
@@ -0,0 +1,205 @@
+// Copyright Epic Games, Inc. All Rights Reserved.
+
+#pragma once
+
+#include "CoreMinimal.h"
+#include "IMessageContext.h"
+#include "Containers/SpscQueue.h"
+#include "Containers/Ticker.h"
+#include "Widgets/Clients/Browser/Models/ITransferStatisticsModel.h"
+
+namespace UE::MultiUserServer
+{
+	template<typename TStatBase>
+	class TClientTransferStatTracker
+	{
+		struct FIncompleteMessageData
+		{
+			uint64 BytesTransferredSoFar;
+			FDateTime LastUpdate;
+		};
+		
+		struct TStatItem : TStatBase
+		{
+			FDateTime LocalTime;
+
+			TStatItem(const FDateTime& LocalTime, TStatBase Data)
+				: TStatBase(Data)
+				, LocalTime(LocalTime)
+			{}
+		};
+
+		using FMessageId = decltype(FOutboundTransferStatistics::MessageId);
+		
+	public:
+
+		using FShouldInclude = TFunction<bool(const TStatBase&)>;
+		using FStatGetter = TFunction<uint64(const TStatBase&)>;
+
+		TClientTransferStatTracker(FShouldInclude ShouldIncludeFunc, FStatGetter StatGetterFunc)
+			: ShouldIncludeFunc(MoveTemp(ShouldIncludeFunc))
+			, StatGetterFunc(MoveTemp(StatGetterFunc))
+		{}
+
+		void OnTransferUpdatedFromThread(TStatBase TransferStatistics)
+		{
+			if (ShouldIncludeFunc(TransferStatistics))
+			{
+				AsyncStatQueue.Enqueue(TStatItem{ FDateTime::Now(), TransferStatistics });
+			}
+		}
+
+		void Tick()
+		{
+			const bool bRemovedOldStats = RemoveOldStatTimelines();
+			const bool bRemovedAnyGroups = RemoveOldGroupedStats();
+			const bool bUpdatedAnyTimeline = !AsyncStatQueue.IsEmpty();
+			
+			const FDateTime Now = FDateTime::Now();
+			FConcertTransferSamplePoint SamplingThisTick{ Now };
+			while (const TOptional<TStatItem> TransferStatistics = AsyncStatQueue.Dequeue())
+			{
+				UpdateStatTimeline(TransferStatistics.GetValue(), SamplingThisTick);
+				UpdateGroupedStats(TransferStatistics.GetValue());
+			}
+
+			if (bRemovedOldStats || bUpdatedAnyTimeline || bRemovedAnyGroups)
+			{
+				TransferStatisticsTimeline.Add(SamplingThisTick);
+				OnTimelineUpdatedDelegates.Broadcast();
+				
+				OnGroupsUpdatedDelegate.Broadcast();
+			}
+		}
+		
+		FOnTransferTimelineUpdated& GetOnTimelineUpdatedDelegates() { return OnTimelineUpdatedDelegates; }
+		FOnTransferGroupsUpdated& GetOnGroupsUpdatedDelegate() { return OnGroupsUpdatedDelegate; }
+		const TArray<FConcertTransferSamplePoint>& GetTransferStatisticsTimeline() const { return TransferStatisticsTimeline; }
+		const TArray<TSharedPtr<TStatBase>>& GetTransferStatisticsGroupedById() const { return TransferStatisticsGroupedById; }
+
+	private:
+		
+		TSpscQueue<TStatItem> AsyncStatQueue;
+		
+		TArray<FConcertTransferSamplePoint> TransferStatisticsTimeline;
+		TMap<FMessageId, FIncompleteMessageData> IncompleteStatsUntilNow;
+		
+		TArray<TSharedPtr<TStatBase>> TransferStatisticsGroupedById;
+		TMap<FMessageId, FDateTime> LastUpdateGroupUpdates;
+		
+		FOnTransferTimelineUpdated OnTimelineUpdatedDelegates;
+		FOnTransferGroupsUpdated OnGroupsUpdatedDelegate;
+
+		FShouldInclude ShouldIncludeFunc;
+		FStatGetter StatGetterFunc;
+		
+		static bool IsTooOld(const FDateTime& StatCreationTime, const FDateTime& Now)
+		{
+			const FTimespan RetainTime = FTimespan::FromSeconds(60.f);
+			return RetainTime < Now - StatCreationTime;
+		}
+		
+		bool RemoveOldStatTimelines()
+		{
+			bool bModifiedTimelines = false;
+			
+			const FDateTime Now = FDateTime::Now();
+			for (auto SampleIt = TransferStatisticsTimeline.CreateIterator(); SampleIt; ++SampleIt)
+			{
+				if (IsTooOld(SampleIt->LocalTime, Now))
+				{
+					SampleIt.RemoveCurrent();
+					bModifiedTimelines = true;
+				}
+			}
+
+			for (auto IncompleteIt = IncompleteStatsUntilNow.CreateIterator(); IncompleteIt; ++IncompleteIt)
+			{
+				if (IsTooOld(IncompleteIt->Value.LastUpdate, Now))
+				{
+					IncompleteIt.RemoveCurrent();
+				}
+			}
+
+			return bModifiedTimelines;
+		}
+		
+		bool RemoveOldGroupedStats()
+		{
+			bool bRemovedAny = false;
+			const FDateTime Now = FDateTime::Now();
+			for (auto It = TransferStatisticsGroupedById.CreateIterator(); It; ++It)
+			{
+				const TSharedPtr<TStatBase>& Stat = *It;
+				if (IsTooOld(LastUpdateGroupUpdates[Stat->MessageId], Now))
+				{
+					LastUpdateGroupUpdates.Remove(Stat->MessageId);
+					It.RemoveCurrent();
+					bRemovedAny = true;
+				}
+			}
+			return bRemovedAny;
+		}
+		
+		void UpdateStatTimeline(const TStatBase& TransferStatistics, FConcertTransferSamplePoint& SamplingThisTick)
+		{
+			uint64 BytesSentSinceLastTime = StatGetterFunc(TransferStatistics);
+			// Same message ID may have been queued multiple times
+			if (FIncompleteMessageData* ItemThisTick = IncompleteStatsUntilNow.Find(TransferStatistics.MessageId))
+			{
+				BytesSentSinceLastTime -= ItemThisTick->BytesTransferredSoFar;
+			}
+
+			IncompleteStatsUntilNow.Add(TransferStatistics.MessageId, { StatGetterFunc(TransferStatistics), FDateTime::Now() });
+			SamplingThisTick.BytesTransferred += BytesSentSinceLastTime;
+		}
+		
+		void UpdateGroupedStats(const TStatBase& TransferStatistics)
+		{
+			const TSharedPtr<TStatBase> NewValue = MakeShared<TStatBase>(TransferStatistics);
+			const auto Pos = Algo::LowerBound(TransferStatisticsGroupedById, NewValue, [](const TSharedPtr<TStatBase>& Value, const TSharedPtr<TStatBase>& Check)
+			{
+				return Value->MessageId < Check->MessageId;
+			});
+			if (TransferStatisticsGroupedById.Num() > 0 && Pos < TransferStatisticsGroupedById.Num() && TransferStatisticsGroupedById[Pos] && TransferStatisticsGroupedById[Pos]->MessageId == NewValue->MessageId)
+			{
+				*TransferStatisticsGroupedById[Pos] = *NewValue;
+			}
+			else
+			{
+				TransferStatisticsGroupedById.Insert(NewValue, Pos);
+			}
+
+			LastUpdateGroupUpdates.Add(NewValue->MessageId, FDateTime::Now());
+		}
+	};
+
+	class FTransferStatisticsModelBase : public ITransferStatisticsModel
+	{
+	public:
+		
+		FTransferStatisticsModelBase();
+		virtual ~FTransferStatisticsModelBase() override;
+		
+		virtual const TArray<FConcertTransferSamplePoint>& GetTransferStatTimeline(EConcertTransferStatistic StatisticType) const override;
+		virtual const TArray<TSharedPtr<FOutboundTransferStatistics>>& GetOutboundTransferStatsGroupedById() const override { return OutboundStatTracker.GetTransferStatisticsGroupedById(); }
+		virtual const TArray<TSharedPtr<FInboundTransferStatistics>>& GetInboundTransferStatsGroupedById() const override { return InboundStatTracker.GetTransferStatisticsGroupedById(); }
+		virtual FOnTransferTimelineUpdated& OnTransferTimelineUpdated(EConcertTransferStatistic StatisticType) override;
+		virtual FOnTransferGroupsUpdated& OnOutboundTransferGroupsUpdated() override { return OutboundStatTracker.GetOnGroupsUpdatedDelegate(); }
+		virtual FOnTransferGroupsUpdated& OnInboundTransferGroupsUpdated() override { return InboundStatTracker.GetOnGroupsUpdatedDelegate(); }
+
+	private:
+		
+		virtual bool ShouldIncludeOutboundStat(const FOutboundTransferStatistics& Item) const = 0;
+		virtual bool ShouldIncludeInboundStat(const FInboundTransferStatistics& Item) const = 0;
+
+		FTSTicker::FDelegateHandle TickHandle;
+		
+		
+		bool Tick(float DeltaTime);
+
+		TClientTransferStatTracker<FOutboundTransferStatistics> OutboundStatTracker;
+		TClientTransferStatTracker<FInboundTransferStatistics> InboundStatTracker;
+	};
+}
+
diff --git a/Engine/Plugins/Developer/Concert/ConcertApp/MultiUserServer/Source/MultiUserServer/Private/Widgets/Clients/Browser/SConcertBrowserItem.cpp b/Engine/Plugins/Developer/Concert/ConcertApp/MultiUserServer/Source/MultiUserServer/Private/Widgets/Clients/Browser/SConcertBrowserItem.cpp
--- a/Engine/Plugins/Developer/Concert/ConcertApp/MultiUserServer/Source/MultiUserServer/Private/Widgets/Clients/Browser/SConcertBrowserItem.cpp
+++ b/Engine/Plugins/Developer/Concert/ConcertApp/MultiUserServer/Source/MultiUserServer/Private/Widgets/Clients/Browser/SConcertBrowserItem.cpp
@@ -2,16 +2,15 @@
 
 #include "SConcertBrowserItem.h"
 
-#include "ClientBrowserItem.h"
 #include "ConcertServerStyle.h"
-#include "IConcertBrowserItem.h"
 #include "INetworkMessagingExtension.h"
 #include "SClientNetworkStats.h"
 #include "Graph/SClientNetworkGraphs.h"
-#include "Models/ClientTransferStatisticsModel.h"
 #include "Models/IClientNetworkStatisticsModel.h"
+#include "Models/Transfer/ClientTransferStatisticsModel.h"
 #include "Table/SClientOutboundTransferStatTable.h"
 #include "Table/SClientInboundTransferStatTable.h"
+#include "Widgets/Clients/Browser/Item/IConcertBrowserItem.h"
 
 #include "Styling/StyleColors.h"
 #include "Widgets/Images/SImage.h"
diff --git a/Engine/Plugins/Developer/Concert/ConcertApp/MultiUserServer/Source/MultiUserServer/Private/Widgets/Clients/Browser/SConcertBrowserItem.h b/Engine/Plugins/Developer/Concert/ConcertApp/MultiUserServer/Source/MultiUserServer/Private/Widgets/Clients/Browser/SConcertBrowserItem.h
--- a/Engine/Plugins/Developer/Concert/ConcertApp/MultiUserServer/Source/MultiUserServer/Private/Widgets/Clients/Browser/SConcertBrowserItem.h
+++ b/Engine/Plugins/Developer/Concert/ConcertApp/MultiUserServer/Source/MultiUserServer/Private/Widgets/Clients/Browser/SConcertBrowserItem.h
@@ -4,7 +4,7 @@
 
 #include "CoreMinimal.h"
 #include "IMessageContext.h"
-#include "Models/IClientTransferStatisticsModel.h"
+#include "Models/ITransferStatisticsModel.h"
 #include "Widgets/DeclarativeSyntaxSupport.h"
 #include "Widgets/SCompoundWidget.h"
 
diff --git a/Engine/Plugins/Developer/Concert/ConcertApp/MultiUserServer/Source/MultiUserServer/Private/Widgets/Clients/Browser/SConcertNetworkBrowser.cpp b/Engine/Plugins/Developer/Concert/ConcertApp/MultiUserServer/Source/MultiUserServer/Private/Widgets/Clients/Browser/SConcertNetworkBrowser.cpp
--- a/Engine/Plugins/Developer/Concert/ConcertApp/MultiUserServer/Source/MultiUserServer/Private/Widgets/Clients/Browser/SConcertNetworkBrowser.cpp
+++ b/Engine/Plugins/Developer/Concert/ConcertApp/MultiUserServer/Source/MultiUserServer/Private/Widgets/Clients/Browser/SConcertNetworkBrowser.cpp
@@ -2,15 +2,14 @@
 
 #include "SConcertNetworkBrowser.h"
 
-#include "ClientBrowserItem.h"
 #include "ConcertServerStyle.h"
 #include "MultiUserServerModule.h"
 #include "Models/IClientBrowserModel.h"
-#include "Models/IClientNetworkStatisticsModel.h"
 #include "SConcertBrowserItem.h"
+#include "Widgets/Clients/Browser/Item/ClientBrowserItem.h"
+
 #include "Algo/Accumulate.h"
 #include "Algo/AllOf.h"
-
 #include "Algo/AnyOf.h"
 #include "Algo/Count.h"
 #include "Dialog/SMessageDialog.h"
@@ -25,14 +24,17 @@
 
 #define LOCTEXT_NAMESPACE "UnrealMultiUserUI.SConcertClientBrowser"
 
-void UE::MultiUserServer::SConcertNetworkBrowser::Construct(const FArguments& InArgs, TSharedRef<IClientBrowserModel> InBrowserModel)
+void UE::MultiUserServer::SConcertNetworkBrowser::Construct(const FArguments& InArgs, TSharedRef<IClientBrowserModel> InBrowserModel, TSharedRef<IClientNetworkStatisticsModel> NetworkStatisticsModel, TSharedRef<IConcertServer> Server)
 {
 	BrowserModel = MoveTemp(InBrowserModel);
-	
+
+	ServerItem = MakeShared<FServerBrowserItem>(Server, NetworkStatisticsModel);
+		
 	HighlightText = MakeShared<FText>();
-	SessionFilter = MakeShared<FClientTextFilter>(FClientTextFilter::FItemToStringArray::CreateSP(this, &SConcertNetworkBrowser::GenerateSearchTerms));
+	SessionFilter = MakeShared<FBrowserItemTextFilter>(FBrowserItemTextFilter::FItemToStringArray::CreateSP(this, &SConcertNetworkBrowser::GenerateSearchTerms));
 	SessionFilter->OnChanged().AddSP(this, &SConcertNetworkBrowser::UpdateTileViewFromAllowedSessions);
 
+	OnServerDoubleClicked = InArgs._OnServerDoubleClicked;
 	OnClientDoubleClicked = InArgs._OnClientDoubleClicked;
 
 	ChildSlot
@@ -61,7 +63,7 @@ void UE::MultiUserServer::SConcertNetworkBrowser::Construct(const FArguments& In
 				SNew(SBox)
 				.HAlign(HAlign_Center)
 				.VAlign(VAlign_Center)
-				.Visibility_Lambda([this](){ return DisplayedClients.Num() == 0 ? EVisibility::Visible : EVisibility::Hidden; })
+				.Visibility_Lambda([this](){ return DisplayedItems.Num() == 0 ? EVisibility::Visible : EVisibility::Hidden; })
 				[
 					SNew(STextBlock)
 					.Text(this, &SConcertNetworkBrowser::GetErrorMessageText)
@@ -192,9 +194,9 @@ TSharedRef<SWidget> UE::MultiUserServer::SConcertNetworkBrowser::CreateTileView(
 {
 	constexpr int32 Height = 270;
 	constexpr int32 Width = 270;
-	return SAssignNew(TileView, STileView<TSharedPtr<FClientBrowserItem>>)
+	return SAssignNew(TileView, STileView<TSharedPtr<IConcertBrowserItem>>)
 		.SelectionMode(ESelectionMode::Multi)
-		.ListItemsSource(&DisplayedClients)
+		.ListItemsSource(&DisplayedItems)
 		.OnGenerateTile(this, &SConcertNetworkBrowser::MakeTileViewWidget)
 		.OnContextMenuOpening(this, &SConcertNetworkBrowser::OnGetContextMenuContent)
 		.OnMouseButtonDoubleClick(this, &SConcertNetworkBrowser::OnListMouseButtonDoubleClick)
@@ -228,11 +230,13 @@ void UE::MultiUserServer::SConcertNetworkBrowser::OnClientListChanged(TSharedPtr
 	case IClientBrowserModel::EClientUpdateType::Added:
 		if (PassesFilter(Item))
 		{
-			DisplayedClients.Add(Item);
+			DisplayedClientItems.Add(Item);
+			DisplayedItems.Add(Item);
 		}
 		break;
 	case IClientBrowserModel::EClientUpdateType::Removed:
-		DisplayedClients.Remove(Item);
+		DisplayedClientItems.Remove(Item);
+		DisplayedItems.Remove(Item);
 		break;
 	default:
 		checkNoEntry();
@@ -331,7 +335,7 @@ FText UE::MultiUserServer::SConcertNetworkBrowser::GetErrorMessageText() const
 	return LOCTEXT("NoClients", "No known clients");
 }
 
-TSharedRef<ITableRow> UE::MultiUserServer::SConcertNetworkBrowser::MakeTileViewWidget(TSharedPtr<FClientBrowserItem> ClientItem, const TSharedRef<STableViewBase>& OwnerTable)
+TSharedRef<ITableRow> UE::MultiUserServer::SConcertNetworkBrowser::MakeTileViewWidget(TSharedPtr<IConcertBrowserItem> ClientItem, const TSharedRef<STableViewBase>& OwnerTable)
 {
 	return SNew(STableRow<TSharedPtr<FClientBrowserItem>>, OwnerTable )
 		.Padding(2.f)
@@ -354,8 +358,8 @@ TSharedPtr<SWidget> UE::MultiUserServer::SConcertNetworkBrowser::OnGetContextMen
 		FUIAction(
 			FExecuteAction::CreateLambda([this]()
 			{
-				TArray<TSharedPtr<FClientBrowserItem>> Items = TileView->GetSelectedItems();
-				for (const TSharedPtr<FClientBrowserItem>& Item : Items)
+				TArray<TSharedPtr<IConcertBrowserItem>> Items = TileView->GetSelectedItems();
+				for (const TSharedPtr<IConcertBrowserItem>& Item : Items)
 				{
 					OnListMouseButtonDoubleClick(Item);
 				}
@@ -391,8 +395,8 @@ void UE::MultiUserServer::SConcertNetworkBrowser::AddDisplayModeEntry(FMenuBuild
 		FUIAction(
 			FExecuteAction::CreateLambda([this, DisplayMode]()
 			{
-				TArray<TSharedPtr<FClientBrowserItem>> Items = TileView->GetSelectedItems();
-				for (const TSharedPtr<FClientBrowserItem>& Item : Items)
+				TArray<TSharedPtr<IConcertBrowserItem>> Items = TileView->GetSelectedItems();
+				for (const TSharedPtr<IConcertBrowserItem>& Item : Items)
 				{
 					Item->SetDisplayMode(DisplayMode);
 				}
@@ -400,7 +404,7 @@ void UE::MultiUserServer::SConcertNetworkBrowser::AddDisplayModeEntry(FMenuBuild
 			FCanExecuteAction::CreateLambda([] { return true; }),
 			FIsActionChecked::CreateLambda([this, DisplayMode]()
 			{
-				return Algo::AllOf(TileView->GetSelectedItems(), [this, DisplayMode](const TSharedPtr<FClientBrowserItem>& Item)
+				return Algo::AllOf(TileView->GetSelectedItems(), [this, DisplayMode](const TSharedPtr<IConcertBrowserItem>& Item)
 				{
 					return Item->GetDisplayMode() == DisplayMode;
 				});
@@ -410,9 +414,20 @@ void UE::MultiUserServer::SConcertNetworkBrowser::AddDisplayModeEntry(FMenuBuild
 	);
 }
 
-void UE::MultiUserServer::SConcertNetworkBrowser::OnListMouseButtonDoubleClick(TSharedPtr<FClientBrowserItem> ClientItem)
+void UE::MultiUserServer::SConcertNetworkBrowser::OnListMouseButtonDoubleClick(TSharedPtr<IConcertBrowserItem> Item)
 {
-	OnClientDoubleClicked.ExecuteIfBound(ClientItem->GetMessageNodeId());
+	const int32 Index = DisplayedClientItems.IndexOfByPredicate([&Item](const TSharedPtr<FClientBrowserItem>& ClientItem)
+	{
+		return Item == ClientItem;
+	});
+	if (Index != INDEX_NONE)
+	{
+		OnClientDoubleClicked.ExecuteIfBound(DisplayedClientItems[Index]->GetMessageNodeId());
+	}
+	else
+	{
+		OnServerDoubleClicked.ExecuteIfBound();
+	}
 }
 
 void UE::MultiUserServer::SConcertNetworkBrowser::SConcertNetworkBrowser::AllowAllSessions()
@@ -430,18 +445,31 @@ void UE::MultiUserServer::SConcertNetworkBrowser::DisallowAllSessions()
 void UE::MultiUserServer::SConcertNetworkBrowser::UpdateTileViewFromAllowedSessions()
 {
 	bShowAllSessions = AllowedSessions.Num() == BrowserModel->GetSessions().Num();
-	DisplayedClients.Empty();
+	DisplayedClientItems.Empty();
+	DisplayedItems.Empty();
+	
+	if (PassesFilter(ServerItem))
+	{
+		DisplayedItems.Add(ServerItem);
+	}
 	
 	for (const TSharedPtr<FClientBrowserItem>& Item : BrowserModel->GetItems())
 	{
 		if (PassesFilter(Item))
 		{
-			DisplayedClients.Add(Item);
+			DisplayedClientItems.Add(Item);
+			DisplayedItems.Add(Item);
 		}
 	}
+	
 	TileView->RequestListRefresh();
 }
 
+bool UE::MultiUserServer::SConcertNetworkBrowser::PassesFilter(const TSharedPtr<FServerBrowserItem>& Client) const
+{
+	return SessionFilter->PassesFilter(Client);
+}
+
 bool UE::MultiUserServer::SConcertNetworkBrowser::PassesFilter(const TSharedPtr<FClientBrowserItem>& Client) const
 {
 	const bool bIsSessionAllowed = Client->GetCurrentSession().IsSet() && AllowedSessions.Contains(*Client->GetCurrentSession());
@@ -449,7 +477,7 @@ bool UE::MultiUserServer::SConcertNetworkBrowser::PassesFilter(const TSharedPtr<
 	return bAllowBaseOnSession && SessionFilter->PassesFilter(Client);
 }
 
-void UE::MultiUserServer::SConcertNetworkBrowser::GenerateSearchTerms(const TSharedPtr<FClientBrowserItem>& Client, TArray<FString>& SearchTerms) const
+void UE::MultiUserServer::SConcertNetworkBrowser::GenerateSearchTerms(const TSharedPtr<IConcertBrowserItem>& Client, TArray<FString>& SearchTerms) const
 {
 	Client->AppendSearchTerms(SearchTerms);
 }
diff --git a/Engine/Plugins/Developer/Concert/ConcertApp/MultiUserServer/Source/MultiUserServer/Private/Widgets/Clients/Browser/SConcertNetworkBrowser.h b/Engine/Plugins/Developer/Concert/ConcertApp/MultiUserServer/Source/MultiUserServer/Private/Widgets/Clients/Browser/SConcertNetworkBrowser.h
--- a/Engine/Plugins/Developer/Concert/ConcertApp/MultiUserServer/Source/MultiUserServer/Private/Widgets/Clients/Browser/SConcertNetworkBrowser.h
+++ b/Engine/Plugins/Developer/Concert/ConcertApp/MultiUserServer/Source/MultiUserServer/Private/Widgets/Clients/Browser/SConcertNetworkBrowser.h
@@ -3,13 +3,12 @@
 #pragma once
 
 #include "CoreMinimal.h"
-#include "SConcertBrowserItem.h"
+#include "Item/ServerBrowserItem.h"
 #include "Misc/TextFilter.h"
 #include "Models/IClientBrowserModel.h"
 #include "Widgets/DeclarativeSyntaxSupport.h"
 #include "Widgets/SCompoundWidget.h"
 
-class FConcertClientsTabController;
 class FMenuBuilder;
 class STableViewBase;
 class ITableRow;
@@ -21,6 +20,7 @@ struct FConcertSessionClientInfo;
 namespace UE::MultiUserServer
 {
 	class FClientBrowserItem;
+	class FServerBrowserItem;
 	class IClientNetworkStatisticsModel;
 	class IClientBrowserModel;
 	enum class EConcertBrowserItemDisplayMode : uint8;
@@ -34,14 +34,16 @@ namespace UE::MultiUserServer
 	public:
 
 		DECLARE_DELEGATE_OneParam(FOnClientDoubleClicked, const FGuid& /*EndpointId*/);
+		DECLARE_DELEGATE(FOnServerDoubleClicked);
 
 		SLATE_BEGIN_ARGS(SConcertNetworkBrowser) {}
 			/** Extension point to the right of the search bar */
 			SLATE_NAMED_SLOT(FArguments, RightOfSearch)
+			SLATE_EVENT(FOnServerDoubleClicked, OnServerDoubleClicked)
 			SLATE_EVENT(FOnClientDoubleClicked, OnClientDoubleClicked)
 		SLATE_END_ARGS()
 
-		void Construct(const FArguments& InArgs, TSharedRef<IClientBrowserModel> InBrowserModel);
+		void Construct(const FArguments& InArgs, TSharedRef<IClientBrowserModel> InBrowserModel, TSharedRef<IClientNetworkStatisticsModel> NetworkStatisticsModel, TSharedRef<IConcertServer> Server);
 
 		/** Shows only the clients connected to the given session ID */
 		void ShowOnlyClientsFromSession(const FGuid& SessionId);
@@ -50,7 +52,7 @@ namespace UE::MultiUserServer
 
 		using FSessionId = FGuid;
 		using FMessagingNodeId = FGuid;
-		using FClientTextFilter = TTextFilter<const TSharedPtr<FClientBrowserItem>&>;
+		using FBrowserItemTextFilter = TTextFilter<const TSharedPtr<IConcertBrowserItem>&>;
 
 		/** Retrieves clients and live sessions */
 		TSharedPtr<IClientBrowserModel> BrowserModel;
@@ -62,15 +64,19 @@ namespace UE::MultiUserServer
 		/** Should admin endpoints be shown? */
 		bool bShowSessionlessClients = true;
 
+		TSharedPtr<FServerBrowserItem> ServerItem;
+		/** Filtered version of IClientBrowserModel::GetItems. Added to DisplayItems. */
+		TArray<TSharedPtr<FClientBrowserItem>> DisplayedClientItems;
 		/** Source array for TileView - fitlered version of IClientBrowserModel::GetItems */
-		TArray<TSharedPtr<FClientBrowserItem>> DisplayedClients;
+		TArray<TSharedPtr<IConcertBrowserItem>> DisplayedItems;
 		/** Visualizes all the items */
-		TSharedPtr<STileView<TSharedPtr<FClientBrowserItem>>> TileView;
+		TSharedPtr<STileView<TSharedPtr<IConcertBrowserItem>>> TileView;
 
 		// Filtering
 		TSharedPtr<FText> HighlightText;
-		TSharedPtr<FClientTextFilter> SessionFilter;
+		TSharedPtr<FBrowserItemTextFilter> SessionFilter;
 		
+		FOnServerDoubleClicked OnServerDoubleClicked;
 		FOnClientDoubleClicked OnClientDoubleClicked;
 
 		TSharedRef<SWidget> CreateSearchArea(const FArguments& InArgs);
@@ -88,17 +94,18 @@ namespace UE::MultiUserServer
 		FText GetErrorMessageText() const;
 		
 		// TileView events
-		TSharedRef<ITableRow> MakeTileViewWidget(TSharedPtr<FClientBrowserItem> ClientItem, const TSharedRef<STableViewBase>& OwnerTable);
+		TSharedRef<ITableRow> MakeTileViewWidget(TSharedPtr<IConcertBrowserItem> ClientItem, const TSharedRef<STableViewBase>& OwnerTable);
 		TSharedPtr<SWidget> OnGetContextMenuContent();
 		void AddDisplayModeEntry(FMenuBuilder& MenuBuilder, EConcertBrowserItemDisplayMode DisplayMode, FText Title, FText Tooltip) const;
-		void OnListMouseButtonDoubleClick(TSharedPtr<FClientBrowserItem> ClientItem);
+		void OnListMouseButtonDoubleClick(TSharedPtr<IConcertBrowserItem> ClientItem);
 
 		// Filtering
 		void AllowAllSessions();
 		void DisallowAllSessions();
 		void UpdateTileViewFromAllowedSessions();
+		bool PassesFilter(const TSharedPtr<FServerBrowserItem>& Client) const;
 		bool PassesFilter(const TSharedPtr<FClientBrowserItem>& Client) const;
-		void GenerateSearchTerms(const TSharedPtr<FClientBrowserItem>& Client, TArray<FString>& SearchTerms) const;
+		void GenerateSearchTerms(const TSharedPtr<IConcertBrowserItem>& Client, TArray<FString>& SearchTerms) const;
 	};
 }
 
diff --git a/Engine/Plugins/Developer/Concert/ConcertApp/MultiUserServer/Source/MultiUserServer/Private/Widgets/Clients/Browser/Table/SClientInboundTransferStatTable.cpp b/Engine/Plugins/Developer/Concert/ConcertApp/MultiUserServer/Source/MultiUserServer/Private/Widgets/Clients/Browser/Table/SClientInboundTransferStatTable.cpp
--- a/Engine/Plugins/Developer/Concert/ConcertApp/MultiUserServer/Source/MultiUserServer/Private/Widgets/Clients/Browser/Table/SClientInboundTransferStatTable.cpp
+++ b/Engine/Plugins/Developer/Concert/ConcertApp/MultiUserServer/Source/MultiUserServer/Private/Widgets/Clients/Browser/Table/SClientInboundTransferStatTable.cpp
@@ -3,7 +3,7 @@
 #include "SClientInboundTransferStatTable.h"
 
 #include "INetworkMessagingExtension.h"
-#include "Widgets/Clients/Browser/Models/IClientTransferStatisticsModel.h"
+#include "Widgets/Clients/Browser/Models/ITransferStatisticsModel.h"
 #include "Widgets/SBoxPanel.h"
 #include "Widgets/SNullWidget.h"
 #include "Widgets/Text/STextBlock.h"
@@ -80,7 +80,7 @@ namespace UE::MultiUserServer
 		}
 	};
 	
-	void SClientInboundTransferStatTable::Construct(const FArguments& InArgs, TSharedRef<IClientTransferStatisticsModel> InStatsModel)
+	void SClientInboundTransferStatTable::Construct(const FArguments& InArgs, TSharedRef<ITransferStatisticsModel> InStatsModel)
 	{
 		StatsModel = InStatsModel;
 
diff --git a/Engine/Plugins/Developer/Concert/ConcertApp/MultiUserServer/Source/MultiUserServer/Private/Widgets/Clients/Browser/Table/SClientInboundTransferStatTable.h b/Engine/Plugins/Developer/Concert/ConcertApp/MultiUserServer/Source/MultiUserServer/Private/Widgets/Clients/Browser/Table/SClientInboundTransferStatTable.h
--- a/Engine/Plugins/Developer/Concert/ConcertApp/MultiUserServer/Source/MultiUserServer/Private/Widgets/Clients/Browser/Table/SClientInboundTransferStatTable.h
+++ b/Engine/Plugins/Developer/Concert/ConcertApp/MultiUserServer/Source/MultiUserServer/Private/Widgets/Clients/Browser/Table/SClientInboundTransferStatTable.h
@@ -13,7 +13,7 @@ struct FOutboundTransferStatistics;
 
 namespace UE::MultiUserServer
 {
-	class IClientTransferStatisticsModel;
+	class ITransferStatisticsModel;
 
 	/**
 	 * Displays FOutboundTransferStatistics as they are updated by INetworkMessagingExtension.
@@ -26,12 +26,12 @@ namespace UE::MultiUserServer
 		{}
 		SLATE_END_ARGS()
 
-		void Construct(const FArguments& InArgs, TSharedRef<IClientTransferStatisticsModel> InStatsModel);
+		void Construct(const FArguments& InArgs, TSharedRef<ITransferStatisticsModel> InStatsModel);
 
 	private:
 
 		/** Tells us when the transfer stats have changed */
-		TSharedPtr<IClientTransferStatisticsModel> StatsModel;
+		TSharedPtr<ITransferStatisticsModel> StatsModel;
 		/** Displays the transfer stats */
 		TSharedPtr<SListView<TSharedPtr<FInboundTransferStatistics>>> SegmenterListView;
 
diff --git a/Engine/Plugins/Developer/Concert/ConcertApp/MultiUserServer/Source/MultiUserServer/Private/Widgets/Clients/Browser/Table/SClientOutboundTransferStatTable.cpp b/Engine/Plugins/Developer/Concert/ConcertApp/MultiUserServer/Source/MultiUserServer/Private/Widgets/Clients/Browser/Table/SClientOutboundTransferStatTable.cpp
--- a/Engine/Plugins/Developer/Concert/ConcertApp/MultiUserServer/Source/MultiUserServer/Private/Widgets/Clients/Browser/Table/SClientOutboundTransferStatTable.cpp
+++ b/Engine/Plugins/Developer/Concert/ConcertApp/MultiUserServer/Source/MultiUserServer/Private/Widgets/Clients/Browser/Table/SClientOutboundTransferStatTable.cpp
@@ -3,7 +3,7 @@
 #include "SClientOutboundTransferStatTable.h"
 
 #include "INetworkMessagingExtension.h"
-#include "Widgets/Clients/Browser/Models/IClientTransferStatisticsModel.h"
+#include "Widgets/Clients/Browser/Models/ITransferStatisticsModel.h"
 #include "Widgets/SBoxPanel.h"
 #include "Widgets/SNullWidget.h"
 #include "Widgets/Text/STextBlock.h"
@@ -97,7 +97,7 @@ namespace UE::MultiUserServer
 		}
 	};
 	
-	void SClientOutboundTransferStatTable::Construct(const FArguments& InArgs, TSharedRef<IClientTransferStatisticsModel> InStatsModel)
+	void SClientOutboundTransferStatTable::Construct(const FArguments& InArgs, TSharedRef<ITransferStatisticsModel> InStatsModel)
 	{
 		StatsModel = InStatsModel;
 
diff --git a/Engine/Plugins/Developer/Concert/ConcertApp/MultiUserServer/Source/MultiUserServer/Private/Widgets/Clients/Browser/Table/SClientOutboundTransferStatTable.h b/Engine/Plugins/Developer/Concert/ConcertApp/MultiUserServer/Source/MultiUserServer/Private/Widgets/Clients/Browser/Table/SClientOutboundTransferStatTable.h
--- a/Engine/Plugins/Developer/Concert/ConcertApp/MultiUserServer/Source/MultiUserServer/Private/Widgets/Clients/Browser/Table/SClientOutboundTransferStatTable.h
+++ b/Engine/Plugins/Developer/Concert/ConcertApp/MultiUserServer/Source/MultiUserServer/Private/Widgets/Clients/Browser/Table/SClientOutboundTransferStatTable.h
@@ -12,7 +12,7 @@ struct FOutboundTransferStatistics;
 
 namespace UE::MultiUserServer
 {
-	class IClientTransferStatisticsModel;
+	class ITransferStatisticsModel;
 
 	/**
 	 * Displays FOutboundTransferStatistics as they are updated by INetworkMessagingExtension.
@@ -25,12 +25,12 @@ namespace UE::MultiUserServer
 		{}
 		SLATE_END_ARGS()
 
-		void Construct(const FArguments& InArgs, TSharedRef<IClientTransferStatisticsModel> InStatsModel);
+		void Construct(const FArguments& InArgs, TSharedRef<ITransferStatisticsModel> InStatsModel);
 
 	private:
 
 		/** Tells us when the transfer stats have changed */
-		TSharedPtr<IClientTransferStatisticsModel> StatsModel;
+		TSharedPtr<ITransferStatisticsModel> StatsModel;
 		/** Displays the transfer stats */
 		TSharedPtr<SListView<TSharedPtr<FOutboundTransferStatistics>>> SegmenterListView;
 
diff --git a/Engine/Plugins/Developer/Concert/ConcertApp/MultiUserServer/Source/MultiUserServer/Private/Widgets/Clients/SConcertClientsTabView.cpp b/Engine/Plugins/Developer/Concert/ConcertApp/MultiUserServer/Source/MultiUserServer/Private/Widgets/Clients/SConcertClientsTabView.cpp
--- a/Engine/Plugins/Developer/Concert/ConcertApp/MultiUserServer/Source/MultiUserServer/Private/Widgets/Clients/SConcertClientsTabView.cpp
+++ b/Engine/Plugins/Developer/Concert/ConcertApp/MultiUserServer/Source/MultiUserServer/Private/Widgets/Clients/SConcertClientsTabView.cpp
@@ -143,16 +143,20 @@ void SConcertClientsTabView::CreateTabs(const TSharedRef<FTabManager>& InTabMana
 
 TSharedRef<SDockTab> SConcertClientsTabView::SpawnClientBrowserTab(const FSpawnTabArgs& InTabArgs)
 {
+	using namespace UE::MultiUserServer;
+	const TSharedRef<FClientNetworkStatisticsModel> NetworkStatisticsModel = MakeShared<FClientNetworkStatisticsModel>();
+	
 	return SNew(SDockTab)
 		.Label(LOCTEXT("ClientBrowserTabLabel", "Clients"))
 		.TabRole(PanelTab)
 		[
 			SAssignNew(ClientBrowser, UE::MultiUserServer::SConcertNetworkBrowser,
-				MakeShared<UE::MultiUserServer::FClientBrowserModel>(Server->GetConcertServer(), ClientInfoCache.ToSharedRef(), MakeShared<UE::MultiUserServer::FClientNetworkStatisticsModel>()))
+				MakeShared<UE::MultiUserServer::FClientBrowserModel>(Server->GetConcertServer(), ClientInfoCache.ToSharedRef(), NetworkStatisticsModel), NetworkStatisticsModel, Server->GetConcertServer())
 			.RightOfSearch()
 			[
 				CreateOpenGlobalLogButton()
 			]
+			.OnServerDoubleClicked_Raw(this, &SConcertClientsTabView::OpenGlobalLogTab)
 			.OnClientDoubleClicked_Raw(this, &SConcertClientsTabView::OpenClientLogTab)
 		]; 
 }
