Always display server stats - even when there are not no clients connected

diff --git a/Engine/Plugins/Developer/Concert/ConcertApp/MultiUserServer/Source/MultiUserServer/Private/Widgets/Clients/Browser/Item/ServerBrowserItem.cpp b/Engine/Plugins/Developer/Concert/ConcertApp/MultiUserServer/Source/MultiUserServer/Private/Widgets/Clients/Browser/Item/ServerBrowserItem.cpp
--- a/Engine/Plugins/Developer/Concert/ConcertApp/MultiUserServer/Source/MultiUserServer/Private/Widgets/Clients/Browser/Item/ServerBrowserItem.cpp
+++ b/Engine/Plugins/Developer/Concert/ConcertApp/MultiUserServer/Source/MultiUserServer/Private/Widgets/Clients/Browser/Item/ServerBrowserItem.cpp
@@ -56,11 +56,6 @@ namespace UE::MultiUserServer
 				AllStats.Add(*Stats);
 			}
 		}
-		
-		if (AllStats.Num() == 0)
-		{
-			return {};
-		}
 
 		FMessageTransportStatistics Result = SumUpStats(AllStats);
 		Result.IPv4AsString = TEXT("localhost"); // Should probably look up the IP address values ...
@@ -85,7 +80,7 @@ namespace UE::MultiUserServer
 			Result.AverageRTT += Stats.AverageRTT;
 		}
 
-		if (ensure(StatArray.Num() > 0))
+		if (StatArray.Num() > 0)
 		{
 			Result.AverageRTT /= StatArray.Num();
 		}
