Improve network statistics by caching network stats so they can be shown when a client disconnected (instead of showing n/a)

diff --git a/Engine/Plugins/Developer/Concert/ConcertApp/MultiUserServer/Source/MultiUserServer/Private/Widgets/Clients/Browser/Item/ServerBrowserItem.cpp b/Engine/Plugins/Developer/Concert/ConcertApp/MultiUserServer/Source/MultiUserServer/Private/Widgets/Clients/Browser/Item/ServerBrowserItem.cpp
--- a/Engine/Plugins/Developer/Concert/ConcertApp/MultiUserServer/Source/MultiUserServer/Private/Widgets/Clients/Browser/Item/ServerBrowserItem.cpp
+++ b/Engine/Plugins/Developer/Concert/ConcertApp/MultiUserServer/Source/MultiUserServer/Private/Widgets/Clients/Browser/Item/ServerBrowserItem.cpp
@@ -40,7 +40,7 @@ namespace UE::MultiUserServer
 			{
 				const FMessageAddress ClientAddress = Session->GetClientAddress(ClientInfo.ClientEndpointId);
 				const TOptional<FMessageTransportStatistics> Stats = NetworkStatisticsModel->GetLatestNetworkStatistics(ClientAddress);
-				if (ensureMsgf(Stats, TEXT("Server stats should always be available because the connection is server controlled.")))
+				if (Stats)
 				{
 					AllStats.Add(*Stats);
 				}
@@ -51,7 +51,7 @@ namespace UE::MultiUserServer
 		for (const FConcertEndpointContext& RemoteClientAdminEndpoints : Server->GetRemoteAdminEndpoints())
 		{
 			const TOptional<FMessageTransportStatistics> Stats = NetworkStatisticsModel->GetLatestNetworkStatistics(Server->GetRemoteAddress(RemoteClientAdminEndpoints.EndpointId));
-			if (ensureMsgf(Stats, TEXT("Server stats should always be available because the connection is server controlled.")))
+			if (Stats)
 			{
 				AllStats.Add(*Stats);
 			}
diff --git a/Engine/Plugins/Developer/Concert/ConcertApp/MultiUserServer/Source/MultiUserServer/Private/Widgets/Clients/Browser/SClientNetworkStats.cpp b/Engine/Plugins/Developer/Concert/ConcertApp/MultiUserServer/Source/MultiUserServer/Private/Widgets/Clients/Browser/SClientNetworkStats.cpp
--- a/Engine/Plugins/Developer/Concert/ConcertApp/MultiUserServer/Source/MultiUserServer/Private/Widgets/Clients/Browser/SClientNetworkStats.cpp
+++ b/Engine/Plugins/Developer/Concert/ConcertApp/MultiUserServer/Source/MultiUserServer/Private/Widgets/Clients/Browser/SClientNetworkStats.cpp
@@ -35,18 +35,9 @@ namespace UE::MultiUserServer
 		{
 			UpdateStatistics(*Stats);
 		}
-		else
-		{
-			const FText NotApplicableText = FText::FromString(TEXT("n/a"));
-			SendText->SetText(NotApplicableText);
-			ReceiveText->SetText(NotApplicableText);
-			RoundTripTimeText->SetText(NotApplicableText);
-			InflightText->SetText(NotApplicableText);
-			LossText->SetText(NotApplicableText);
-		}
 	}
 
-	void SClientNetworkStats::UpdateStatistics(const FMessageTransportStatistics& Statistics)
+	void SClientNetworkStats::UpdateStatistics(const FMessageTransportStatistics& Statistics) const
 	{
 		SendText->SetText(FText::FromString(NetworkStatistics::FormatTotalBytesSent(Statistics)));
 		ReceiveText->SetText(FText::FromString(NetworkStatistics::FormatTotalBytesReceived(Statistics)));
diff --git a/Engine/Plugins/Developer/Concert/ConcertApp/MultiUserServer/Source/MultiUserServer/Private/Widgets/Clients/Browser/SClientNetworkStats.h b/Engine/Plugins/Developer/Concert/ConcertApp/MultiUserServer/Source/MultiUserServer/Private/Widgets/Clients/Browser/SClientNetworkStats.h
--- a/Engine/Plugins/Developer/Concert/ConcertApp/MultiUserServer/Source/MultiUserServer/Private/Widgets/Clients/Browser/SClientNetworkStats.h
+++ b/Engine/Plugins/Developer/Concert/ConcertApp/MultiUserServer/Source/MultiUserServer/Private/Widgets/Clients/Browser/SClientNetworkStats.h
@@ -44,7 +44,7 @@ namespace UE::MultiUserServer
 		TSharedPtr<STextBlock> InflightText;
 		TSharedPtr<STextBlock> LossText;
 		
-		void UpdateStatistics(const FMessageTransportStatistics& Statistics);
+		void UpdateStatistics(const FMessageTransportStatistics& Statistics) const;
 
 		TSharedRef<SWidget> CreateStatTable();
 		void AddStatistic(const TSharedRef<SHorizontalBox>& AddTo, const FText& StatisticName, const FText& StatisticToolTip, TSharedPtr<STextBlock>& AssignTo);
diff --git a/Engine/Plugins/Developer/Concert/ConcertApp/MultiUserServer/Source/MultiUserServer/Private/Widgets/Clients/Browser/SConcertBrowserItem.cpp b/Engine/Plugins/Developer/Concert/ConcertApp/MultiUserServer/Source/MultiUserServer/Private/Widgets/Clients/Browser/SConcertBrowserItem.cpp
--- a/Engine/Plugins/Developer/Concert/ConcertApp/MultiUserServer/Source/MultiUserServer/Private/Widgets/Clients/Browser/SConcertBrowserItem.cpp
+++ b/Engine/Plugins/Developer/Concert/ConcertApp/MultiUserServer/Source/MultiUserServer/Private/Widgets/Clients/Browser/SConcertBrowserItem.cpp
@@ -84,6 +84,16 @@ void UE::MultiUserServer::SConcertBrowserItem::Construct(const FArguments& InArg
 	];
 }
 
+void UE::MultiUserServer::SConcertBrowserItem::Tick(const FGeometry& AllottedGeometry, const double InCurrentTime, const float InDeltaTime)
+{
+	SCompoundWidget::Tick(AllottedGeometry, InCurrentTime, InDeltaTime);
+
+	if (const TOptional<FMessageTransportStatistics> Stats = Item->GetLatestNetworkStatistics())
+	{
+		LastAvailableIp = NetworkStatistics::FormatIPv4AsString(Stats);
+	}
+}
+
 TSharedRef<SWidget> UE::MultiUserServer::SConcertBrowserItem::CreateHeader()
 {
 	return SAssignNew(ClientName, STextBlock)
@@ -162,14 +172,20 @@ TSharedRef<SWidget> UE::MultiUserServer::SConcertBrowserItem::CreateFooter()
 				SAssignNew(ClientIP4, STextBlock)
 				.ColorAndOpacity(FColor::White)
 				.HighlightText_Lambda([this](){ return *HighlightText; })
-				.Text_Lambda([this]()
-				{
-					return FText::FromString(NetworkStatistics::FormatIPv4AsString(Item->GetLatestNetworkStatistics()));
-				})
+				.Text(this, &SConcertBrowserItem::GetIpText)
 			]
 		];
 }
 
+FText UE::MultiUserServer::SConcertBrowserItem::GetIpText() const
+{
+	if (LastAvailableIp)
+	{
+		return FText::FromString(*LastAvailableIp);
+	}
+	return FText::FromString(NetworkStatistics::FormatIPv4AsString({}));
+}
+
 const FSlateBrush* UE::MultiUserServer::SConcertBrowserItem::GetBackgroundImage() const
 {
 	if (IsHovered())
diff --git a/Engine/Plugins/Developer/Concert/ConcertApp/MultiUserServer/Source/MultiUserServer/Private/Widgets/Clients/Browser/SConcertBrowserItem.h b/Engine/Plugins/Developer/Concert/ConcertApp/MultiUserServer/Source/MultiUserServer/Private/Widgets/Clients/Browser/SConcertBrowserItem.h
--- a/Engine/Plugins/Developer/Concert/ConcertApp/MultiUserServer/Source/MultiUserServer/Private/Widgets/Clients/Browser/SConcertBrowserItem.h
+++ b/Engine/Plugins/Developer/Concert/ConcertApp/MultiUserServer/Source/MultiUserServer/Private/Widgets/Clients/Browser/SConcertBrowserItem.h
@@ -26,6 +26,8 @@ namespace UE::MultiUserServer
 		SLATE_END_ARGS()
 		
 		void Construct(const FArguments& InArgs, TSharedRef<IConcertBrowserItem> InClientItem);
+		/** Updates LastAvailableIp */
+		virtual void Tick(const FGeometry& AllottedGeometry, const double InCurrentTime, const float InDeltaTime) override;
 
 	private:
 
@@ -38,12 +40,17 @@ namespace UE::MultiUserServer
 		TSharedPtr<STextBlock> ClientName;
 		TSharedPtr<SClientNetworkStats> NetworkStats;
 		TSharedPtr<STextBlock> ClientIP4;
+
+		/** Stores the IP address once it becomes it available. Used so we can continue displaying it when the client goes offline. */
+		TOptional<FString> LastAvailableIp;
 		
 		TSharedRef<SWidget> CreateHeader();
 		TSharedRef<SWidget> CreateContentArea();
 		TSharedRef<SWidget> CreateStats();
 		TSharedRef<SWidget> CreateFooter();
 
+		FText GetIpText() const;
+
 		const FSlateBrush* GetBackgroundImage() const;
 	};
 }
