Visualize send and received packages

diff --git a/Engine/Plugins/Developer/Concert/ConcertApp/MultiUserServer/Source/MultiUserServer/Private/ConcertServerStyle.cpp b/Engine/Plugins/Developer/Concert/ConcertApp/MultiUserServer/Source/MultiUserServer/Private/ConcertServerStyle.cpp
--- a/Engine/Plugins/Developer/Concert/ConcertApp/MultiUserServer/Source/MultiUserServer/Private/ConcertServerStyle.cpp
+++ b/Engine/Plugins/Developer/Concert/ConcertApp/MultiUserServer/Source/MultiUserServer/Private/ConcertServerStyle.cpp
@@ -72,6 +72,9 @@ TSharedRef<FSlateStyleSet> FConcertServerStyle::Create()
 	StyleSet->Set("Concert.Ack.Ack", new IMAGE_PLUGIN_BRUSH("Icons/Ack_Ack_16x", Icon16x16));
 	StyleSet->Set("Concert.Ack.Success", new IMAGE_PLUGIN_BRUSH("Icons/Ack_Success_16x", Icon16x16));
 	StyleSet->Set("Concert.Ack.Failure", new IMAGE_PLUGIN_BRUSH("Icons/Ack_Fail_16x", Icon16x16));
+	// Icons: Package transmission
+	StyleSet->Set("Concert.PackageTransmission.Success", new IMAGE_PLUGIN_BRUSH("Icons/Ack_Success_16x", Icon16x16));
+	StyleSet->Set("Concert.PackageTransmission.Failure", new IMAGE_PLUGIN_BRUSH("Icons/Ack_Fail_16x", Icon16x16));
 	// Icons: Muting
 	StyleSet->Set("Concert.Muted",     new IMAGE_PLUGIN_BRUSH("Icons/Muted_16x", Icon16x16));
 	StyleSet->Set("Concert.Unmuted",     new IMAGE_PLUGIN_BRUSH("Icons/Unmuted_16x", Icon16x16));
diff --git a/Engine/Plugins/Developer/Concert/ConcertApp/MultiUserServer/Source/MultiUserServer/Private/Settings/MultiUserServerColumnVisibilitySettings.cpp b/Engine/Plugins/Developer/Concert/ConcertApp/MultiUserServer/Source/MultiUserServer/Private/Settings/MultiUserServerColumnVisibilitySettings.cpp
--- a/Engine/Plugins/Developer/Concert/ConcertApp/MultiUserServer/Source/MultiUserServer/Private/Settings/MultiUserServerColumnVisibilitySettings.cpp
+++ b/Engine/Plugins/Developer/Concert/ConcertApp/MultiUserServer/Source/MultiUserServer/Private/Settings/MultiUserServerColumnVisibilitySettings.cpp
@@ -35,6 +35,10 @@ UMultiUserServerColumnVisibilitySettings::UMultiUserServerColumnVisibilitySettin
 	{
 		SaveConfig();
 	});
+	OnOnPackageTransmissionColumnVisibilityChanged().AddLambda([this](const FColumnVisibilitySnapshot&)
+	{
+		SaveConfig();
+	});
 
 	FCoreDelegates::OnPreExit.AddLambda([]()
 	{
diff --git a/Engine/Plugins/Developer/Concert/ConcertApp/MultiUserServer/Source/MultiUserServer/Private/Settings/MultiUserServerColumnVisibilitySettings.h b/Engine/Plugins/Developer/Concert/ConcertApp/MultiUserServer/Source/MultiUserServer/Private/Settings/MultiUserServerColumnVisibilitySettings.h
--- a/Engine/Plugins/Developer/Concert/ConcertApp/MultiUserServer/Source/MultiUserServer/Private/Settings/MultiUserServerColumnVisibilitySettings.h
+++ b/Engine/Plugins/Developer/Concert/ConcertApp/MultiUserServer/Source/MultiUserServer/Private/Settings/MultiUserServerColumnVisibilitySettings.h
@@ -45,6 +45,10 @@ public:
 	const FColumnVisibilitySnapshot& GetTransportLogColumnVisibility() const { return TransportLogColumnVisibility; }
 	void SetTransportLogColumnVisibility(FColumnVisibilitySnapshot NewValue) { TransportLogColumnVisibility = MoveTemp(NewValue); OnTransportLogColumnVisibilityEvent.Broadcast(TransportLogColumnVisibility); }
 	FOnColumnVisibilitySnapshotChanged& OnTransportLogColumnVisibility() { return OnTransportLogColumnVisibilityEvent; }
+
+	const FColumnVisibilitySnapshot& GetPackageTransmissionColumnVisibility() const { return PackageTransmissionColumnVisibility; }
+	void SetPackageTransmissionColumnVisibility(FColumnVisibilitySnapshot NewValue) { PackageTransmissionColumnVisibility = MoveTemp(NewValue); OnPackageTransmissionColumnVisibilityChangedEvent.Broadcast(PackageTransmissionColumnVisibility); }
+	FOnColumnVisibilitySnapshotChanged& OnOnPackageTransmissionColumnVisibilityChanged() { return OnPackageTransmissionColumnVisibilityChangedEvent; }
 	
 private:
 	
@@ -71,4 +75,8 @@ private:
 	UPROPERTY(Config)
 	FColumnVisibilitySnapshot TransportLogColumnVisibility;
 	FOnColumnVisibilitySnapshotChanged OnTransportLogColumnVisibilityEvent;
+
+	UPROPERTY(Config)
+	FColumnVisibilitySnapshot PackageTransmissionColumnVisibility;
+	FOnColumnVisibilitySnapshotChanged OnPackageTransmissionColumnVisibilityChangedEvent;
 };
diff --git a/Engine/Plugins/Developer/Concert/ConcertApp/MultiUserServer/Source/MultiUserServer/Private/Settings/MultiUserServerPackageTransmissionSettings.cpp b/Engine/Plugins/Developer/Concert/ConcertApp/MultiUserServer/Source/MultiUserServer/Private/Settings/MultiUserServerPackageTransmissionSettings.cpp
new file mode 100644
--- /dev/null
+++ b/Engine/Plugins/Developer/Concert/ConcertApp/MultiUserServer/Source/MultiUserServer/Private/Settings/MultiUserServerPackageTransmissionSettings.cpp
@@ -0,0 +1,8 @@
+// Copyright Epic Games, Inc. All Rights Reserved.
+
+#include "MultiUserServerPackageTransmissionSettings.h"
+
+UMultiUserServerPackageTransmissionSettings* UMultiUserServerPackageTransmissionSettings::GetSettings()
+{
+	return GetMutableDefault<UMultiUserServerPackageTransmissionSettings>();
+}
\ No newline at end of file
diff --git a/Engine/Plugins/Developer/Concert/ConcertApp/MultiUserServer/Source/MultiUserServer/Private/Settings/MultiUserServerPackageTransmissionSettings.h b/Engine/Plugins/Developer/Concert/ConcertApp/MultiUserServer/Source/MultiUserServer/Private/Settings/MultiUserServerPackageTransmissionSettings.h
new file mode 100644
--- /dev/null
+++ b/Engine/Plugins/Developer/Concert/ConcertApp/MultiUserServer/Source/MultiUserServer/Private/Settings/MultiUserServerPackageTransmissionSettings.h
@@ -0,0 +1,24 @@
+// Copyright Epic Games, Inc. All Rights Reserved.
+
+#pragma once
+
+#include "CoreMinimal.h"
+#include "ConcertFrontendUtils.h"
+#include "UObject/Object.h"
+#include "MultiUserServerPackageTransmissionSettings.generated.h"
+
+/**
+ * 
+ */
+UCLASS(config = MultiUserServerUserSettings, DefaultConfig)
+class MULTIUSERSERVER_API UMultiUserServerPackageTransmissionSettings : public UObject
+{
+	GENERATED_BODY()
+public:
+
+	static UMultiUserServerPackageTransmissionSettings* GetSettings();
+
+	/** How to display FPackageTransmissionEntry::Time */
+	UPROPERTY(Config)
+	ETimeFormat TimestampTimeFormat = ETimeFormat::Absolute;
+};
diff --git a/Engine/Plugins/Developer/Concert/ConcertApp/MultiUserServer/Source/MultiUserServer/Private/Widgets/Clients/Logging/ConcertLogEntry.h b/Engine/Plugins/Developer/Concert/ConcertApp/MultiUserServer/Source/MultiUserServer/Private/Widgets/Clients/Logging/ConcertLogEntry.h
--- a/Engine/Plugins/Developer/Concert/ConcertApp/MultiUserServer/Source/MultiUserServer/Private/Widgets/Clients/Logging/ConcertLogEntry.h
+++ b/Engine/Plugins/Developer/Concert/ConcertApp/MultiUserServer/Source/MultiUserServer/Private/Widgets/Clients/Logging/ConcertLogEntry.h
@@ -57,3 +57,5 @@ struct FConcertLogEntry
 	/** Additional data to display in the UI */
 	FConcertLogMetadata LogMetaData;
 };
+
+using FConcertLogEntryFilterFunc = TFunctionRef<bool(const FConcertLogEntry&)>;
diff --git a/Engine/Plugins/Developer/Concert/ConcertApp/MultiUserServer/Source/MultiUserServer/Private/Widgets/Clients/Logging/Filter/ConcertLogFilter_FrontendRoot.cpp b/Engine/Plugins/Developer/Concert/ConcertApp/MultiUserServer/Source/MultiUserServer/Private/Widgets/Clients/Logging/Filter/ConcertLogFilter_FrontendRoot.cpp
--- a/Engine/Plugins/Developer/Concert/ConcertApp/MultiUserServer/Source/MultiUserServer/Private/Widgets/Clients/Logging/Filter/ConcertLogFilter_FrontendRoot.cpp
+++ b/Engine/Plugins/Developer/Concert/ConcertApp/MultiUserServer/Source/MultiUserServer/Private/Widgets/Clients/Logging/Filter/ConcertLogFilter_FrontendRoot.cpp
@@ -17,18 +17,18 @@
 
 FConcertLogFilter_FrontendRoot::FConcertLogFilter_FrontendRoot(TSharedRef<FConcertLogTokenizer> Tokenizer, TArray<TSharedRef<FConcertFrontendLogFilter>> InCustomFilters, const TArray<TSharedRef<FConcertLogFilter>>& NonVisualFilters)
 	: TextSearchFilter(MakeShared<FConcertFrontendLogFilter_TextSearch>(MoveTemp(Tokenizer)))
-	, CustomFilters(MoveTemp(InCustomFilters))
-	, AllFrontendFilters(CustomFilters)
+	, FrontendFilters(MoveTemp(InCustomFilters))
+	, AllFilters(FrontendFilters)
 {
-	AllFrontendFilters.Reserve(CustomFilters.Num() + 1 + NonVisualFilters.Num());
+	AllFilters.Reserve(FrontendFilters.Num() + 1 + NonVisualFilters.Num());
 	
-	AllFrontendFilters.Add(TextSearchFilter);
+	AllFilters.Add(TextSearchFilter);
 	for (const TSharedRef<FConcertLogFilter>& NonVisualFilter : NonVisualFilters)
 	{
-		AllFrontendFilters.Add(NonVisualFilter);
+		AllFilters.Add(NonVisualFilter);
 	}
 	
-	for (const TSharedRef<FConcertLogFilter>& Filter : AllFrontendFilters)
+	for (const TSharedRef<FConcertLogFilter>& Filter : AllFilters)
 	{
 		Filter->OnChanged().AddLambda([this]()
 		{
@@ -59,7 +59,7 @@ TSharedRef<SWidget> FConcertLogFilter_FrontendRoot::BuildFilterWidgets() const
 bool FConcertLogFilter_FrontendRoot::PassesFilter(const FConcertLogEntry& InItem) const
 {
 	return Algo::AllOf(
-		AllFrontendFilters,
+		AllFilters,
 		[&InItem](const TSharedRef<FConcertLogFilter>& AndFilter){ return AndFilter->PassesFilter(InItem); }
 		);
 }
@@ -69,7 +69,7 @@ TSharedRef<SWidget> FConcertLogFilter_FrontendRoot::BuildCustomFilterListWidget(
 	TSharedRef<SHorizontalBox> Box = SNew(SHorizontalBox);
 
 	bool bIsFirst = true;
-	for (const TSharedRef<FConcertFrontendLogFilter>& Filter : CustomFilters)
+	for (const TSharedRef<FConcertFrontendLogFilter>& Filter : FrontendFilters)
 	{
 		const FMargin Margin = bIsFirst ? FMargin() : FMargin(8, 0, 0, 0);
 		bIsFirst = false;
diff --git a/Engine/Plugins/Developer/Concert/ConcertApp/MultiUserServer/Source/MultiUserServer/Private/Widgets/Clients/Logging/Filter/ConcertLogFilter_FrontendRoot.h b/Engine/Plugins/Developer/Concert/ConcertApp/MultiUserServer/Source/MultiUserServer/Private/Widgets/Clients/Logging/Filter/ConcertLogFilter_FrontendRoot.h
--- a/Engine/Plugins/Developer/Concert/ConcertApp/MultiUserServer/Source/MultiUserServer/Private/Widgets/Clients/Logging/Filter/ConcertLogFilter_FrontendRoot.h
+++ b/Engine/Plugins/Developer/Concert/ConcertApp/MultiUserServer/Source/MultiUserServer/Private/Widgets/Clients/Logging/Filter/ConcertLogFilter_FrontendRoot.h
@@ -35,10 +35,10 @@ private:
 	TSharedRef<FConcertFrontendLogFilter_TextSearch> TextSearchFilter;
 
 	/** AllFrontendFilters without special filters we have as propteries above, such as TextSearchFilter. */
-	TArray<TSharedRef<FConcertFrontendLogFilter>> CustomFilters;
+	TArray<TSharedRef<FConcertFrontendLogFilter>> FrontendFilters;
 	
 	/** Filters that are combined using logical AND. */
-	TArray<TSharedRef<FConcertLogFilter>> AllFrontendFilters;
+	TArray<TSharedRef<FConcertLogFilter>> AllFilters;
 
 	/** Builds the widgets that go under the text */
 	TSharedRef<SWidget> BuildCustomFilterListWidget() const;
diff --git a/Engine/Plugins/Developer/Concert/ConcertApp/MultiUserServer/Source/MultiUserServer/Private/Widgets/Clients/Logging/Filter/FilteredConcertLogList.cpp b/Engine/Plugins/Developer/Concert/ConcertApp/MultiUserServer/Source/MultiUserServer/Private/Widgets/Clients/Logging/Filter/FilteredConcertLogList.cpp
--- a/Engine/Plugins/Developer/Concert/ConcertApp/MultiUserServer/Source/MultiUserServer/Private/Widgets/Clients/Logging/Filter/FilteredConcertLogList.cpp
+++ b/Engine/Plugins/Developer/Concert/ConcertApp/MultiUserServer/Source/MultiUserServer/Private/Widgets/Clients/Logging/Filter/FilteredConcertLogList.cpp
@@ -39,6 +39,7 @@ void FFilteredConcertLogList::OnLowestLogEntryChanged(FConcertLogID NewLowestVal
 	{
 		// Worst case this is O(n^2) - usually a thought O(cn) with small c)...
 		// Not much we without rewriting SListView to use something other than TArray
+		RemoveMessageIdToLogEntry(FilteredLogs[i]);
 		FilteredLogs.RemoveAt(i);
 		--i;
 	}
@@ -48,6 +49,7 @@ void FFilteredConcertLogList::OnNewLogEntryAdded(const TSharedRef<FConcertLogEnt
 {
 	if (!Filter || Filter->PassesFilter(*NewLogEntry))
 	{
+		MessageIdToLogs.FindOrAdd(NewLogEntry->Log.MessageId).AddUnique(NewLogEntry);
 		FilteredLogs.Emplace(NewLogEntry);
 		OnLogListChanged().Broadcast(GetFilteredLogs());
 	}
@@ -55,11 +57,13 @@ void FFilteredConcertLogList::OnNewLogEntryAdded(const TSharedRef<FConcertLogEnt
 
 void FFilteredConcertLogList::RebuildFilteredResult()
 {
+	MessageIdToLogs.Reset();
 	FilteredLogs.Reset();
 	LogSource->ForEachLog([this](const TSharedPtr<FConcertLogEntry>& LogEntry)
 	{
 		if (!Filter || Filter->PassesFilter(*LogEntry))
 		{
+			MessageIdToLogs.FindOrAdd(LogEntry->Log.MessageId).AddUnique(LogEntry);
 			FilteredLogs.Emplace(LogEntry);
 		}
 	});
@@ -67,6 +71,24 @@ void FFilteredConcertLogList::RebuildFilteredResult()
 	OnLogListChanged().Broadcast(GetFilteredLogs());
 }
 
+void FFilteredConcertLogList::RemoveMessageIdToLogEntry(const TSharedPtr<FConcertLogEntry>& Log)
+{
+	TArray<TSharedPtr<FConcertLogEntry>>* Entries = MessageIdToLogs.Find(Log->Log.MessageId);
+	if (!Entries)
+	{
+		return;
+	}
+	
+	if (Entries->Num() == 1)
+	{
+		MessageIdToLogs.Remove(Log->Log.MessageId);
+	}
+	else
+	{
+		Entries->RemoveSingle(Log);
+	}
+}
+
 FPagedFilteredConcertLogList::FPagedFilteredConcertLogList(TSharedRef<IConcertLogSource> LogSource, TSharedPtr<FConcertLogFilter> OptionalFilter, uint16 InitialLogsPerPage)
 	: FFilteredConcertLogList(MoveTemp(LogSource), MoveTemp(OptionalFilter))
 	, LogsPerPage(FMath::Max<uint16>(InitialLogsPerPage, 1))
diff --git a/Engine/Plugins/Developer/Concert/ConcertApp/MultiUserServer/Source/MultiUserServer/Private/Widgets/Clients/Logging/Filter/FilteredConcertLogList.h b/Engine/Plugins/Developer/Concert/ConcertApp/MultiUserServer/Source/MultiUserServer/Private/Widgets/Clients/Logging/Filter/FilteredConcertLogList.h
--- a/Engine/Plugins/Developer/Concert/ConcertApp/MultiUserServer/Source/MultiUserServer/Private/Widgets/Clients/Logging/Filter/FilteredConcertLogList.h
+++ b/Engine/Plugins/Developer/Concert/ConcertApp/MultiUserServer/Source/MultiUserServer/Private/Widgets/Clients/Logging/Filter/FilteredConcertLogList.h
@@ -21,6 +21,11 @@ public:
 
 	bool IsBufferFull() const { return LogSource->IsBufferFull(); }
 	const TArray<TSharedPtr<FConcertLogEntry>>& GetFilteredLogs() const { return FilteredLogs; }
+	TArray<TSharedPtr<FConcertLogEntry>> GetFilteredLogsWithId(const FGuid& MessageId) const
+	{
+		const TArray<TSharedPtr<FConcertLogEntry>>* Result = MessageIdToLogs.Find(MessageId);
+		return Result ? *Result : TArray<TSharedPtr<FConcertLogEntry>>{};
+	}
 
 protected:
 
@@ -36,6 +41,9 @@ private:
 
 	/** Contains the items of LogSource with Filter applied. */
 	TArray<TSharedPtr<FConcertLogEntry>> FilteredLogs;
+
+	/** Binds message IDs to the logs that have that ID. Logs can share the same ID but have different MessageActions. */
+	TMap<FGuid, TArray<TSharedPtr<FConcertLogEntry>>> MessageIdToLogs;
 	
 	/** Called when the content of FilteredResult changes */
 	FOnLogListChanged OnLogListChangedEvent;
@@ -43,6 +51,8 @@ private:
 	void OnLowestLogEntryChanged(FConcertLogID NewLowestValidID);
 	void OnNewLogEntryAdded(const TSharedRef<FConcertLogEntry>& NewLogEntry);
 	void RebuildFilteredResult();
+
+	void RemoveMessageIdToLogEntry(const TSharedPtr<FConcertLogEntry>& Log);
 };
 
 /** Further sorts the results of FFilteredConcertLogList::Filter into pages, i.e. buckets of logs, useful for limiting elements in a table view. */
diff --git a/Engine/Plugins/Developer/Concert/ConcertApp/MultiUserServer/Source/MultiUserServer/Private/Widgets/Clients/Logging/LogScrollingDelegates.h b/Engine/Plugins/Developer/Concert/ConcertApp/MultiUserServer/Source/MultiUserServer/Private/Widgets/Clients/Logging/LogScrollingDelegates.h
new file mode 100644
--- /dev/null
+++ b/Engine/Plugins/Developer/Concert/ConcertApp/MultiUserServer/Source/MultiUserServer/Private/Widgets/Clients/Logging/LogScrollingDelegates.h
@@ -0,0 +1,12 @@
+// Copyright Epic Games, Inc. All Rights Reserved.
+
+#pragma once
+
+#include "CoreMinimal.h"
+#include "Widgets/Clients/Logging/ConcertLogEntry.h"
+
+namespace UE::MultiUserServer
+{
+	DECLARE_DELEGATE_RetVal_ThreeParams(bool, FCanScrollToLog, const FGuid& /*MessageId*/, FConcertLogEntryFilterFunc /*Filter*/, FText& /*ErrorMessage*/);
+	DECLARE_DELEGATE_TwoParams(FScrollToLog, const FGuid& /*MessageId*/, FConcertLogEntryFilterFunc /*Filter*/);
+}
\ No newline at end of file
diff --git a/Engine/Plugins/Developer/Concert/ConcertApp/MultiUserServer/Source/MultiUserServer/Private/Widgets/Clients/Logging/SConcertTransportLog.cpp b/Engine/Plugins/Developer/Concert/ConcertApp/MultiUserServer/Source/MultiUserServer/Private/Widgets/Clients/Logging/SConcertTransportLog.cpp
--- a/Engine/Plugins/Developer/Concert/ConcertApp/MultiUserServer/Source/MultiUserServer/Private/Widgets/Clients/Logging/SConcertTransportLog.cpp
+++ b/Engine/Plugins/Developer/Concert/ConcertApp/MultiUserServer/Source/MultiUserServer/Private/Widgets/Clients/Logging/SConcertTransportLog.cpp
@@ -28,6 +28,7 @@ const FName SConcertTransportLog::FirstColumnId("AvatarColourColumnId");
 SConcertTransportLog::~SConcertTransportLog()
 {
 	ConcertTransportEvents::OnConcertTransportLoggingEnabledChangedEvent().RemoveAll(this);
+	UMultiUserServerColumnVisibilitySettings::GetSettings()->OnTransportLogColumnVisibility().RemoveAll(this);
 }
 
 void SConcertTransportLog::Construct(const FArguments& InArgs, TSharedRef<IConcertLogSource> LogSource, TSharedRef<FEndpointToUserNameCache> InEndpointCache, TSharedRef<FConcertLogTokenizer> InLogTokenizer)
@@ -88,6 +89,24 @@ void SConcertTransportLog::Construct(const FArguments& InArgs, TSharedRef<IConce
 	OnConcertLoggingEnabledChanged(ConcertTransportEvents::IsLoggingEnabled());
 }
 
+bool SConcertTransportLog::CanScrollToLog(const FGuid& MessageId, FConcertLogEntryFilterFunc Filter) const
+{
+	return PagedLogList->GetFilteredLogsWithId(MessageId).ContainsByPredicate([Filter](const TSharedPtr<FConcertLogEntry>& LogEntry)
+	{
+		return Filter(*LogEntry);
+	});
+}
+
+void SConcertTransportLog::ScrollToLog(const FGuid& MessageId, FConcertLogEntryFilterFunc Filter) const
+{
+	// Could be optimised further by having FPagedFilteredConcertLogList cache the index but it's fine - ScrollToLog is called very infrequently (when the user presses a button)
+	const int32 Index = PagedLogList->GetFilteredLogs().IndexOfByPredicate([MessageId, Filter](const TSharedPtr<FConcertLogEntry>& Entry)
+	{
+		return Entry->Log.MessageId == MessageId && Filter(*Entry);
+	});
+	ScrollToLog(Index);
+}
+
 TSharedRef<SWidget> SConcertTransportLog::CreateTableView()
 {
 	return SAssignNew(LogView, SListView<TSharedPtr<FConcertLogEntry>>)
@@ -179,6 +198,8 @@ TSharedRef<ITableRow> SConcertTransportLog::OnGenerateActivityRowWidget(TSharedP
 	
 	return SNew(SConcertTransportLogRow, Item, OwnerTable, LogTokenizer.ToSharedRef(), HighlightText.ToSharedRef())
 		.AvatarColor(AvatarColor)
+		.CanScrollToAckLog(this, &SConcertTransportLog::CanScrollToAckLog)
+		.CanScrollToAckedLog(this, &SConcertTransportLog::CanScrollToAckedLog)
 		.ScrollToAckLog(this, &SConcertTransportLog::ScrollToAckLog)
 		.ScrollToAckedLog(this, &SConcertTransportLog::ScrollToAckedLog);
 }
@@ -345,22 +366,48 @@ void SConcertTransportLog::OnConcertLoggingEnabledChanged(bool bNewEnabled)
 	}
 }
 
+bool SConcertTransportLog::CanScrollToAckLog(const FGuid& MessageId) const
+{
+	return CanScrollToLog(MessageId, [MessageId](const FConcertLogEntry& Entry)
+	{
+		return SharedCanScrollToAckLog(MessageId, Entry);
+	});
+}
+
+bool SConcertTransportLog::CanScrollToAckedLog(const FGuid& MessageId) const
+{
+	return CanScrollToLog(MessageId, [MessageId](const FConcertLogEntry& Entry)
+	{
+		return SharedCanScrollToAckedLog(MessageId, Entry);
+	});
+}
+
 void SConcertTransportLog::ScrollToAckLog(const FGuid& MessageId) const
 {
-	const int32 Index = PagedLogList->GetFilteredLogs().IndexOfByPredicate([MessageId](const TSharedPtr<FConcertLogEntry>& Entry)
+	ScrollToLog(MessageId, [MessageId](const FConcertLogEntry& Entry)
 	{
-		return Entry->Log.MessageId == MessageId && Entry->Log.MessageAction == EConcertLogMessageAction::Process;
+		return SharedCanScrollToAckLog(MessageId, Entry);
 	});
-	ScrollToLog(Index);
 }
 
 void SConcertTransportLog::ScrollToAckedLog(const FGuid& MessageId) const
 {
-	const int32 Index = PagedLogList->GetFilteredLogs().IndexOfByPredicate([MessageId](const TSharedPtr<FConcertLogEntry>& Entry)
+	ScrollToLog(MessageId, [MessageId](const FConcertLogEntry& Entry)
 	{
-		return Entry->Log.MessageId == MessageId && Entry->Log.MessageAction == EConcertLogMessageAction::Send;
+		return SharedCanScrollToAckedLog(MessageId, Entry);
 	});
-	ScrollToLog(Index);
+}
+
+bool SConcertTransportLog::SharedCanScrollToAckLog(const FGuid& MessageId, const FConcertLogEntry& Entry)
+{
+	// MessageId is shared by some logs: there is a log for receiving one one for processing for example.
+	return Entry.Log.MessageId == MessageId && Entry.Log.MessageAction == EConcertLogMessageAction::Process;
+}
+
+bool SConcertTransportLog::SharedCanScrollToAckedLog(const FGuid& MessageId, const FConcertLogEntry& Entry)
+{
+	// MessageId is shared by some logs: there is a log for receiving one one for sending for example.
+	return Entry.Log.MessageId == MessageId && Entry.Log.MessageAction == EConcertLogMessageAction::Send;
 }
 
 void SConcertTransportLog::ScrollToLog(const int32 LogIndex) const
@@ -371,10 +418,6 @@ void SConcertTransportLog::ScrollToLog(const int32 LogIndex) const
 		LogView->RequestNavigateToItem(PagedLogList->GetFilteredLogs()[LogIndex]);
 		LogView->SetSelection(PagedLogList->GetFilteredLogs()[LogIndex]);
 	}
-	else
-	{
-		// TODO: Tell the user. Ideally the button calling would be disabled but that requires an optmised search algorithm for checking whether an item is currently filtered out or not.
-	}
 }
 
 #undef LOCTEXT_NAMESPACE
diff --git a/Engine/Plugins/Developer/Concert/ConcertApp/MultiUserServer/Source/MultiUserServer/Private/Widgets/Clients/Logging/SConcertTransportLog.h b/Engine/Plugins/Developer/Concert/ConcertApp/MultiUserServer/Source/MultiUserServer/Private/Widgets/Clients/Logging/SConcertTransportLog.h
--- a/Engine/Plugins/Developer/Concert/ConcertApp/MultiUserServer/Source/MultiUserServer/Private/Widgets/Clients/Logging/SConcertTransportLog.h
+++ b/Engine/Plugins/Developer/Concert/ConcertApp/MultiUserServer/Source/MultiUserServer/Private/Widgets/Clients/Logging/SConcertTransportLog.h
@@ -41,6 +41,9 @@ public:
 
 	void Construct(const FArguments& InArgs, TSharedRef<IConcertLogSource> LogSource, TSharedRef<FEndpointToUserNameCache> InEndpointCache, TSharedRef<FConcertLogTokenizer> InLogTokenizer);
 
+	bool CanScrollToLog(const FGuid& MessageId, FConcertLogEntryFilterFunc Filter) const;
+	void ScrollToLog(const FGuid& MessageId, FConcertLogEntryFilterFunc Filter) const;
+	
 private:
 	
 	/** Used to overlay EnableLoggingPrompt over the tabs */
@@ -84,7 +87,12 @@ private:
 	
 	void OnConcertLoggingEnabledChanged(bool bNewEnabled);
 
+	bool CanScrollToAckLog(const FGuid& MessageId) const;
+	bool CanScrollToAckedLog(const FGuid& MessageId) const;
 	void ScrollToAckLog(const FGuid& MessageId) const;
 	void ScrollToAckedLog(const FGuid& MessageId) const;
+	static bool SharedCanScrollToAckLog(const FGuid& MessageId, const FConcertLogEntry& Entry);
+	static bool SharedCanScrollToAckedLog(const FGuid& MessageId, const FConcertLogEntry& Entry);
+	
 	void ScrollToLog(const int32 LogIndex) const;
 };
diff --git a/Engine/Plugins/Developer/Concert/ConcertApp/MultiUserServer/Source/MultiUserServer/Private/Widgets/Clients/Logging/SConcertTransportLogRow.cpp b/Engine/Plugins/Developer/Concert/ConcertApp/MultiUserServer/Source/MultiUserServer/Private/Widgets/Clients/Logging/SConcertTransportLogRow.cpp
--- a/Engine/Plugins/Developer/Concert/ConcertApp/MultiUserServer/Source/MultiUserServer/Private/Widgets/Clients/Logging/SConcertTransportLogRow.cpp
+++ b/Engine/Plugins/Developer/Concert/ConcertApp/MultiUserServer/Source/MultiUserServer/Private/Widgets/Clients/Logging/SConcertTransportLogRow.cpp
@@ -5,6 +5,8 @@
 #include "ConcertServerStyle.h"
 #include "Util/ConcertLogTokenizer.h"
 #include "SConcertTransportLog.h"
+#include "Algo/AllOf.h"
+#include "Algo/AnyOf.h"
 #include "Session/Activity/PredefinedActivityColumns.h"
 
 #include "Framework/MultiBox/MultiBoxBuilder.h"
@@ -25,6 +27,8 @@ void SConcertTransportLogRow::Construct(const FArguments& InArgs, TSharedPtr<FCo
 	HighlightText = MoveTemp(InHighlightText);
 	
 	AvatarColor = InArgs._AvatarColor;
+	CanScrollToAckLogFunc = InArgs._CanScrollToAckLog;
+	CanScrollToAckedLogFunc = InArgs._CanScrollToAckedLog;
 	ScrollToAckLogFunc = InArgs._ScrollToAckLog;
 	ScrollToAckedLogFunc = InArgs._ScrollToAckedLog;
 	check(ScrollToAckLogFunc.IsBound() && ScrollToAckedLogFunc.IsBound());
@@ -91,26 +95,6 @@ TSharedRef<SWidget> SConcertTransportLogRow::CreateAckColumn(const FName& Proper
 {
 	return SNew(SOverlay)
 
-		.ToolTipText_Lambda([this]()
-		{
-			switch (LogEntry->LogMetaData.AckState)
-			{
-				case EConcertLogAckState::Ack: 
-					return LOCTEXT("Ack.Scroll", "This is an ACK. Press to jump to the log this log ACKs.");
-				case EConcertLogAckState::AckReceived: 
-					return LOCTEXT("Acked.Scroll", "ACK received. Press to jump to the log that ACKs this log");
-				case EConcertLogAckState::AckFailure: 
-					return LOCTEXT("AckFailure", "No ACK received: retries exhausted.");
-				case EConcertLogAckState::NotNeeded:
-					return LOCTEXT("Received", "No ACK needed.");
-				case EConcertLogAckState::InProgress:
-					return LOCTEXT("InProgress", "Awaiting ACK...");
-			default:
-				checkNoEntry();
-				return FText::GetEmpty();
-			}
-		})
-
 		+SOverlay::Slot()
 		.HAlign(HAlign_Center)
 		.VAlign(VAlign_Center)
@@ -161,6 +145,7 @@ TSharedRef<SWidget> SConcertTransportLogRow::CreateAckInProgressWidget() const
 					? EVisibility::Visible
 					: EVisibility::Hidden;
 			})
+			.ToolTipText(LOCTEXT("InProgress", "Awaiting ACK..."))
 		];
 }
 
@@ -168,6 +153,10 @@ TSharedRef<SWidget> SConcertTransportLogRow::CreateAckSuccessWidget() const
 {
 	return SNew(SButton)
 		.ContentPadding(0.f)
+		.IsEnabled_Lambda([this]()
+		{
+			return LogEntry->LogMetaData.AckingMessageId.IsSet() && CanScrollToAckLogFunc.Execute(*LogEntry->LogMetaData.AckingMessageId);
+		})
 		.OnPressed_Lambda([this]()
 		{
 			if (LogEntry->LogMetaData.AckingMessageId.IsSet())
@@ -181,6 +170,12 @@ TSharedRef<SWidget> SConcertTransportLogRow::CreateAckSuccessWidget() const
 				? EVisibility::Visible
 				: EVisibility::Hidden;
 		})
+		.ToolTipText_Lambda([this]()
+		{
+			return LogEntry->LogMetaData.AckingMessageId.IsSet() && CanScrollToAckLogFunc.Execute(*LogEntry->LogMetaData.AckingMessageId)
+				? LOCTEXT("Acked.Scroll.Enabled", "ACK received. Press to jump to the log that ACKs this log")
+				: LOCTEXT("Acked.Scroll.Disabled", "ACK received\nThis button is disabled because the corresponding ACK log is filtered out.");
+		})
 		[
 			SNew(SImage)
 			.DesiredSizeOverride(FVector2D(16.f, 16.f))
@@ -198,7 +193,8 @@ TSharedRef<SWidget> SConcertTransportLogRow::CreateAckFailureWidget() const
 			return LogEntry->LogMetaData.AckState == EConcertLogAckState::AckFailure
 				? EVisibility::Visible
 				: EVisibility::Hidden;
-		});
+		})
+		.ToolTipText(LOCTEXT("AckFailure", "No ACK received: retries exhausted."));
 }
 
 TSharedRef<SWidget> SConcertTransportLogRow::CreateAckNotNeededWidget() const
@@ -210,13 +206,31 @@ TSharedRef<SWidget> SConcertTransportLogRow::CreateAckNotNeededWidget() const
 			return LogEntry->LogMetaData.AckState == EConcertLogAckState::NotNeeded
 				? EVisibility::Visible
 				: EVisibility::Hidden;
-		});
+		})
+		.ToolTipText(LOCTEXT("Received", "No ACK needed."));
 }
 
 TSharedRef<SWidget> SConcertTransportLogRow::CreateAckWidget() const
 {
 	return SNew(SComboButton)
 		.HasDownArrow(false)
+		.ToolTipText_Lambda([this]()
+		{
+			const bool bIsEnabled = LogEntry->LogMetaData.AckedMessageId && Algo::AnyOf(*LogEntry->LogMetaData.AckedMessageId, [this](const FGuid& MessageId)
+				{
+				return CanScrollToAckedLogFunc.Execute(MessageId);
+			});
+			return bIsEnabled
+				? LOCTEXT("Ack.Scroll.Disabled", "This is an ACK. Press to jump to the log this log ACKs.")
+				: LOCTEXT("Ack.Scroll.Enabled", "This is an ACK.\nThis button is disabled because all logs this log ACKs are filtered out.");
+		})
+		.IsEnabled_Lambda([this]()
+		{
+			return LogEntry->LogMetaData.AckedMessageId && Algo::AnyOf(*LogEntry->LogMetaData.AckedMessageId, [this](const FGuid& MessageId)
+			{
+				return CanScrollToAckedLogFunc.Execute(MessageId);
+			});
+		})
 		.OnGetMenuContent_Lambda([this]()
 		{
 			FMenuBuilder MenuBuilder(true, nullptr);
@@ -224,14 +238,14 @@ TSharedRef<SWidget> SConcertTransportLogRow::CreateAckWidget() const
 			{
 				MenuBuilder.AddMenuEntry(
 					FText::FromString(MessageId.ToString(EGuidFormats::DigitsWithHyphens)),
-					FText::GetEmpty(),
+					CanScrollToAckedLogFunc.Execute(MessageId) ? FText() : LOCTEXT("Ack.Scroll.ComoboEntry.Disabled", "This log is filtered out."),
 					FSlateIcon(),
 					FUIAction(
 						FExecuteAction::CreateLambda([this, MessageId]()
 						{
 							ScrollToAckedLogFunc.Execute(MessageId);
 						}),
-						FCanExecuteAction::CreateLambda([] { return true; })),
+						FCanExecuteAction::CreateLambda([this, MessageId] { return CanScrollToAckedLogFunc.Execute(MessageId); })),
 					NAME_None,
 					EUserInterfaceActionType::Button
 				);
diff --git a/Engine/Plugins/Developer/Concert/ConcertApp/MultiUserServer/Source/MultiUserServer/Private/Widgets/Clients/Logging/SConcertTransportLogRow.h b/Engine/Plugins/Developer/Concert/ConcertApp/MultiUserServer/Source/MultiUserServer/Private/Widgets/Clients/Logging/SConcertTransportLogRow.h
--- a/Engine/Plugins/Developer/Concert/ConcertApp/MultiUserServer/Source/MultiUserServer/Private/Widgets/Clients/Logging/SConcertTransportLogRow.h
+++ b/Engine/Plugins/Developer/Concert/ConcertApp/MultiUserServer/Source/MultiUserServer/Private/Widgets/Clients/Logging/SConcertTransportLogRow.h
@@ -14,10 +14,13 @@ class SConcertTransportLogRow : public SMultiColumnTableRow<TSharedPtr<FConcertL
 public:
 
 	DECLARE_DELEGATE_OneParam(FScrollToLog, const FGuid& /*MessageId*/);
+	DECLARE_DELEGATE_RetVal_OneParam(bool, FCanScrollToLog, const FGuid& /*MessageId*/);
 
 	SLATE_BEGIN_ARGS(SConcertTransportLogRow)
 	{}
 		SLATE_ARGUMENT(FLinearColor, AvatarColor)
+		SLATE_EVENT(FCanScrollToLog, CanScrollToAckLog)
+		SLATE_EVENT(FCanScrollToLog, CanScrollToAckedLog)
 		SLATE_EVENT(FScrollToLog, ScrollToAckLog)
 		SLATE_EVENT(FScrollToLog, ScrollToAckedLog)
 	SLATE_END_ARGS()
@@ -38,6 +41,8 @@ private:
 	TSharedPtr<FText> HighlightText;
 
 	FLinearColor AvatarColor;
+	FCanScrollToLog CanScrollToAckLogFunc;
+	FCanScrollToLog CanScrollToAckedLogFunc;
 	FScrollToLog ScrollToAckLogFunc;
 	FScrollToLog ScrollToAckedLogFunc;
 
diff --git a/Engine/Plugins/Developer/Concert/ConcertApp/MultiUserServer/Source/MultiUserServer/Private/Widgets/Clients/Logging/Util/ConcertLogTokenizer.cpp b/Engine/Plugins/Developer/Concert/ConcertApp/MultiUserServer/Source/MultiUserServer/Private/Widgets/Clients/Logging/Util/ConcertLogTokenizer.cpp
--- a/Engine/Plugins/Developer/Concert/ConcertApp/MultiUserServer/Source/MultiUserServer/Private/Widgets/Clients/Logging/Util/ConcertLogTokenizer.cpp
+++ b/Engine/Plugins/Developer/Concert/ConcertApp/MultiUserServer/Source/MultiUserServer/Private/Widgets/Clients/Logging/Util/ConcertLogTokenizer.cpp
@@ -60,12 +60,12 @@ FString FConcertLogTokenizer::TokenizeCustomPayloadUncompressedByteSize(const FC
 
 FString FConcertLogTokenizer::TokenizeOriginEndpointId(const FConcertLog& Data) const
 {
-	return GetEndpointDisplayString(Data.OriginEndpointId);
+	return EndpointInfoGetter->GetEndpointDisplayString(Data.OriginEndpointId);
 }
 
 FString FConcertLogTokenizer::TokenizeDestinationEndpointId(const FConcertLog& Data) const
 {
-	return GetEndpointDisplayString(Data.DestinationEndpointId);
+	return EndpointInfoGetter->GetEndpointDisplayString(Data.DestinationEndpointId);
 }
 
 FString FConcertLogTokenizer::TokenizeUsingPropertyExport(const void* ContainerPtr, const FProperty& ConcertLogProperty) const
@@ -77,18 +77,3 @@ FString FConcertLogTokenizer::TokenizeUsingPropertyExport(const void* ContainerP
 	check(bSuccess);
 	return Exported;
 }
-
-FString FConcertLogTokenizer::GetEndpointDisplayString(const FGuid& EndpointId) const
-{
-	if (EndpointInfoGetter->IsServerEndpoint(EndpointId))
-	{
-		return FString("Server");		
-	}
-
-	if (const TOptional<FConcertClientInfo> ClientInfo = EndpointInfoGetter->GetClientInfo(EndpointId))
-	{
-		return ClientInfo->DisplayName;
-	}
-	
-	return EndpointId.ToString(EGuidFormats::DigitsWithHyphens);
-}
diff --git a/Engine/Plugins/Developer/Concert/ConcertApp/MultiUserServer/Source/MultiUserServer/Private/Widgets/Clients/Logging/Util/ConcertLogTokenizer.h b/Engine/Plugins/Developer/Concert/ConcertApp/MultiUserServer/Source/MultiUserServer/Private/Widgets/Clients/Logging/Util/ConcertLogTokenizer.h
--- a/Engine/Plugins/Developer/Concert/ConcertApp/MultiUserServer/Source/MultiUserServer/Private/Widgets/Clients/Logging/Util/ConcertLogTokenizer.h
+++ b/Engine/Plugins/Developer/Concert/ConcertApp/MultiUserServer/Source/MultiUserServer/Private/Widgets/Clients/Logging/Util/ConcertLogTokenizer.h
@@ -37,8 +37,6 @@ private:
 
 	/** Used so we can look up client and server info (even after client has disconnected) */
 	TSharedRef<FEndpointToUserNameCache> EndpointInfoGetter;
-
-	FString GetEndpointDisplayString(const FGuid& EndpointId) const;
 	
 	FString TokenizeUsingPropertyExport(const void* ContainerPtr, const FProperty& ConcertLogProperty) const;
 };
diff --git a/Engine/Plugins/Developer/Concert/ConcertApp/MultiUserServer/Source/MultiUserServer/Private/Widgets/Clients/PackageTransmission/Filter/FrontendPackageTransmissionFilter.h b/Engine/Plugins/Developer/Concert/ConcertApp/MultiUserServer/Source/MultiUserServer/Private/Widgets/Clients/PackageTransmission/Filter/FrontendPackageTransmissionFilter.h
new file mode 100644
--- /dev/null
+++ b/Engine/Plugins/Developer/Concert/ConcertApp/MultiUserServer/Source/MultiUserServer/Private/Widgets/Clients/PackageTransmission/Filter/FrontendPackageTransmissionFilter.h
@@ -0,0 +1,53 @@
+// Copyright Epic Games, Inc. All Rights Reserved.
+
+#pragma once
+
+#include "CoreMinimal.h"
+#include "PackageTransmissionFilter.h"
+#include "Misc/IFilter.h"
+
+class SWidget;
+
+namespace UE::MultiUserServer
+{
+	/** A filter that is intended to be displayed in the UI. Every filter has one widget displaying it. */
+	class FFrontendPackageTransmissionFilter : public FPackageTransmissionFilter
+	{
+	public:
+
+		/** Gets a widget that represents this filter */
+		virtual TSharedRef<SWidget> GetFilterWidget() = 0;
+	};
+
+	/** Helper class to implement filters */
+	template<typename TConcertLogFilter, typename TWidgetType = SWidget>
+	class TFrontendPackageTransmissionFilterAggregate : public FFrontendPackageTransmissionFilter
+	{
+	public:
+
+		//~ Begin FConcertLogFilter Interface
+		virtual bool PassesFilter(const FPackageTransmissionEntry& InItem) const final override { return Implementation.PassesFilter(InItem); }
+		//~ End FConcertLogFilter Interface
+	
+		//~ Begin FConcertFrontendLogFilter Interface
+		virtual TSharedRef<SWidget> GetFilterWidget() final override { return ChildSlot.ToSharedRef(); }
+		//~ End FConcertLogFilter Interface
+
+	protected:
+
+		template<typename... TArg>
+		TFrontendPackageTransmissionFilterAggregate(TArg&&... Arg)
+			: Implementation(Forward<TArg>(Arg)...)
+		{
+			Implementation.OnChanged().AddLambda([this]()
+			{
+				OnChanged().Broadcast();
+			});
+		}
+
+		TConcertLogFilter Implementation;
+		TSharedPtr<TWidgetType> ChildSlot;
+	};
+}
+
+
diff --git a/Engine/Plugins/Developer/Concert/ConcertApp/MultiUserServer/Source/MultiUserServer/Private/Widgets/Clients/PackageTransmission/Filter/FrontendPackageTransmissionFilter_TextSearch.cpp b/Engine/Plugins/Developer/Concert/ConcertApp/MultiUserServer/Source/MultiUserServer/Private/Widgets/Clients/PackageTransmission/Filter/FrontendPackageTransmissionFilter_TextSearch.cpp
new file mode 100644
--- /dev/null
+++ b/Engine/Plugins/Developer/Concert/ConcertApp/MultiUserServer/Source/MultiUserServer/Private/Widgets/Clients/PackageTransmission/Filter/FrontendPackageTransmissionFilter_TextSearch.cpp
@@ -0,0 +1,41 @@
+// Copyright Epic Games, Inc. All Rights Reserved.
+
+#include "FrontendPackageTransmissionFilter_TextSearch.h"
+
+#include "Widgets/Clients/PackageTransmission/Util/PackageTransmissionEntryTokenizer.h"
+
+namespace UE::MultiUserServer
+{
+	FPackageTransmissionFilter_TextSearch::FPackageTransmissionFilter_TextSearch(TSharedRef<FPackageTransmissionEntryTokenizer> Tokenizer)
+		: TextFilter(TTextFilter<const FPackageTransmissionEntry&>::FItemToStringArray::CreateRaw(this, &FPackageTransmissionFilter_TextSearch::GenerateSearchTerms))
+		, Tokenizer(MoveTemp(Tokenizer))
+	{
+		TextFilter.OnChanged().AddLambda([this]
+		{
+			OnChanged().Broadcast();
+		});
+	}
+
+	void FPackageTransmissionFilter_TextSearch::GenerateSearchTerms(const FPackageTransmissionEntry& InItem, TArray<FString>& OutTerms) const
+	{
+		OutTerms.Add(Tokenizer->TokenizeTime(InItem));
+		OutTerms.Add(Tokenizer->TokenizeOrigin(InItem));
+		OutTerms.Add(Tokenizer->TokenizeDestination(InItem));
+		OutTerms.Add(Tokenizer->TokenizeSize(InItem));
+		OutTerms.Add(Tokenizer->TokenizeRevision(InItem));
+		OutTerms.Add(Tokenizer->TokenizePackagePath(InItem));
+		OutTerms.Add(Tokenizer->TokenizePackageName(InItem));
+	}
+
+	FFrontendPackageTransmissionFilter_TextSearch::FFrontendPackageTransmissionFilter_TextSearch(TSharedRef<FPackageTransmissionEntryTokenizer> Tokenizer)
+		: Super(MoveTemp(Tokenizer))
+	{
+		ChildSlot = SNew(SSearchBox)
+		.OnTextChanged_Lambda([this](const FText& NewSearchText)
+		{
+			Implementation.SetRawFilterText(NewSearchText);
+			SearchText = NewSearchText;
+		})
+		.DelayChangeNotificationsWhileTyping(true);
+	}
+}
diff --git a/Engine/Plugins/Developer/Concert/ConcertApp/MultiUserServer/Source/MultiUserServer/Private/Widgets/Clients/PackageTransmission/Filter/FrontendPackageTransmissionFilter_TextSearch.h b/Engine/Plugins/Developer/Concert/ConcertApp/MultiUserServer/Source/MultiUserServer/Private/Widgets/Clients/PackageTransmission/Filter/FrontendPackageTransmissionFilter_TextSearch.h
new file mode 100644
--- /dev/null
+++ b/Engine/Plugins/Developer/Concert/ConcertApp/MultiUserServer/Source/MultiUserServer/Private/Widgets/Clients/PackageTransmission/Filter/FrontendPackageTransmissionFilter_TextSearch.h
@@ -0,0 +1,51 @@
+// Copyright Epic Games, Inc. All Rights Reserved.
+
+#pragma once
+
+#include "CoreMinimal.h"
+#include "FrontendPackageTransmissionFilter.h"
+#include "Misc/TextFilter.h"
+#include "Widgets/Input/SSearchBox.h"
+
+namespace UE::MultiUserServer
+{
+	class FPackageTransmissionEntryTokenizer;
+
+	class FPackageTransmissionFilter_TextSearch : public FPackageTransmissionFilter
+	{
+	public:
+		
+		FPackageTransmissionFilter_TextSearch(TSharedRef<FPackageTransmissionEntryTokenizer> Tokenizer);
+
+		//~ Begin FPackageTransmissionFilter Interface
+		virtual bool PassesFilter(const FPackageTransmissionEntry& InItem) const override { return TextFilter.PassesFilter(InItem); }
+		//~ End FPackageTransmissionFilter Interface
+		
+		void SetRawFilterText(const FText& InFilterText) { TextFilter.SetRawFilterText(InFilterText); }
+	
+	private:
+
+		/** Does the actual string search */
+		TTextFilter<const FPackageTransmissionEntry&> TextFilter;
+		/** Helps in converting FConcertLog members into search terms */
+		TSharedRef<FPackageTransmissionEntryTokenizer> Tokenizer;
+
+		/** Parses InItem into a bunch of strings that can be searched */
+		void GenerateSearchTerms(const FPackageTransmissionEntry& InItem, TArray<FString>& OutTerms) const;
+	};
+	
+    class FFrontendPackageTransmissionFilter_TextSearch : public TFrontendPackageTransmissionFilterAggregate<FPackageTransmissionFilter_TextSearch, SSearchBox>
+    {
+		using Super = TFrontendPackageTransmissionFilterAggregate<FPackageTransmissionFilter_TextSearch, SSearchBox>;
+    public:
+    	
+    	FFrontendPackageTransmissionFilter_TextSearch(TSharedRef<FPackageTransmissionEntryTokenizer> Tokenizer);
+
+    	FText GetSearchText() const { return SearchText; }
+
+    private:
+
+    	FText SearchText;
+    };
+}
+
diff --git a/Engine/Plugins/Developer/Concert/ConcertApp/MultiUserServer/Source/MultiUserServer/Private/Widgets/Clients/PackageTransmission/Filter/PackageTransmissionFilter.h b/Engine/Plugins/Developer/Concert/ConcertApp/MultiUserServer/Source/MultiUserServer/Private/Widgets/Clients/PackageTransmission/Filter/PackageTransmissionFilter.h
new file mode 100644
--- /dev/null
+++ b/Engine/Plugins/Developer/Concert/ConcertApp/MultiUserServer/Source/MultiUserServer/Private/Widgets/Clients/PackageTransmission/Filter/PackageTransmissionFilter.h
@@ -0,0 +1,37 @@
+// Copyright Epic Games, Inc. All Rights Reserved.
+
+#pragma once
+
+#include "CoreMinimal.h"
+#include "Misc/IFilter.h"
+
+namespace UE::MultiUserServer
+{
+	struct FPackageTransmissionEntry;
+
+	/** Base filter */
+	class FPackageTransmissionFilter
+		:
+		public IFilter<const FPackageTransmissionEntry&>,
+		public TSharedFromThis<FPackageTransmissionFilter>,
+		// We do not need to copy filters so let's avoid it by accident; some constructors pass this pointer to callbacks which become stale upon copying
+		public FNoncopyable
+	{
+	public:
+
+		//~ Begin IFilter Interface
+		DECLARE_DERIVED_EVENT( FFrontendFilter, IFilter<const FPackageTransmissionEntry&>::FChangedEvent, FChangedEvent );
+		virtual FChangedEvent& OnChanged() override { return ChangedEvent; }
+		//~ End IFilter Interface
+
+		protected:
+	
+		void BroadcastChangedEvent() const { ChangedEvent.Broadcast(); }
+
+	private:
+	
+		FChangedEvent ChangedEvent;
+	};
+}
+
+
diff --git a/Engine/Plugins/Developer/Concert/ConcertApp/MultiUserServer/Source/MultiUserServer/Private/Widgets/Clients/PackageTransmission/Filter/PackageTransmissionFilter_FrontendRoot.cpp b/Engine/Plugins/Developer/Concert/ConcertApp/MultiUserServer/Source/MultiUserServer/Private/Widgets/Clients/PackageTransmission/Filter/PackageTransmissionFilter_FrontendRoot.cpp
new file mode 100644
--- /dev/null
+++ b/Engine/Plugins/Developer/Concert/ConcertApp/MultiUserServer/Source/MultiUserServer/Private/Widgets/Clients/PackageTransmission/Filter/PackageTransmissionFilter_FrontendRoot.cpp
@@ -0,0 +1,89 @@
+// Copyright Epic Games, Inc. All Rights Reserved.
+
+#include "PackageTransmissionFilter_FrontendRoot.h"
+
+#include "Widgets/Clients/PackageTransmission/Model/PackageTransmissionEntry.h"
+#include "Widgets/Clients/PackageTransmission/Filter/FrontendPackageTransmissionFilter_TextSearch.h"
+
+#include "Algo/AllOf.h"
+#include "Widgets/SBoxPanel.h"
+
+namespace UE::MultiUserServer
+{
+	FPackageTransmissionFilter_FrontendRoot::FPackageTransmissionFilter_FrontendRoot(
+		TSharedRef<FPackageTransmissionEntryTokenizer> Tokenizer,
+		TArray<TSharedRef<FFrontendPackageTransmissionFilter>> InCustomFilters,
+		const TArray<TSharedRef<FPackageTransmissionFilter>>& NonVisualFilters)
+		: TextSearchFilter(MakeShared<FFrontendPackageTransmissionFilter_TextSearch>(MoveTemp(Tokenizer)))
+		, FrontendFilters(MoveTemp(InCustomFilters))
+		, AllFilters(FrontendFilters)
+	{
+		AllFilters.Add(TextSearchFilter);
+		for (const TSharedRef<FPackageTransmissionFilter>& NonVisual : NonVisualFilters)
+		{
+			AllFilters.Add(NonVisual);
+		}
+	
+		for (const TSharedRef<FPackageTransmissionFilter>& Filter : AllFilters)
+		{
+			Filter->OnChanged().AddLambda([this]()
+			{
+				OnChanged().Broadcast();
+			});
+		}
+	}
+
+	TSharedRef<SWidget> FPackageTransmissionFilter_FrontendRoot::BuildFilterWidgets() const
+	{
+		return SNew(SVerticalBox)
+
+			// Search bar
+			+SVerticalBox::Slot()
+			.AutoHeight()
+			[
+				TextSearchFilter->GetFilterWidget()
+			]
+
+			+SVerticalBox::Slot()
+			.AutoHeight()
+			.Padding(0, 2)
+			[
+				BuildCustomFilterListWidget()
+			];
+	}
+
+	bool FPackageTransmissionFilter_FrontendRoot::PassesFilter(const FPackageTransmissionEntry& InItem) const
+	{
+		return Algo::AllOf(
+			AllFilters,
+			[&InItem](const TSharedRef<FPackageTransmissionFilter>& AndFilter){ return AndFilter->PassesFilter(InItem); }
+		);
+	}
+
+	TSharedRef<SWidget> FPackageTransmissionFilter_FrontendRoot::BuildCustomFilterListWidget() const
+	{
+		TSharedRef<SHorizontalBox> Box = SNew(SHorizontalBox);
+
+		bool bIsFirst = true;
+		for (const TSharedRef<FFrontendPackageTransmissionFilter>& Filter : FrontendFilters)
+		{
+			const FMargin Margin = bIsFirst ? FMargin() : FMargin(8, 0, 0, 0);
+			bIsFirst = false;
+			Box->AddSlot()
+				.AutoWidth()
+				.Padding(Margin)
+				.VAlign(VAlign_Center)
+				[
+					Filter->GetFilterWidget()
+				];
+		}
+
+		return Box;
+	}
+
+	TSharedRef<FPackageTransmissionFilter_FrontendRoot> MakeFilter(TSharedRef<FPackageTransmissionEntryTokenizer> Tokenizer)
+	{
+		TArray<TSharedRef<FFrontendPackageTransmissionFilter>> FrontendFilters;
+		return MakeShared<FPackageTransmissionFilter_FrontendRoot>(MoveTemp(Tokenizer), FrontendFilters);
+	}
+}
diff --git a/Engine/Plugins/Developer/Concert/ConcertApp/MultiUserServer/Source/MultiUserServer/Private/Widgets/Clients/PackageTransmission/Filter/PackageTransmissionFilter_FrontendRoot.h b/Engine/Plugins/Developer/Concert/ConcertApp/MultiUserServer/Source/MultiUserServer/Private/Widgets/Clients/PackageTransmission/Filter/PackageTransmissionFilter_FrontendRoot.h
new file mode 100644
--- /dev/null
+++ b/Engine/Plugins/Developer/Concert/ConcertApp/MultiUserServer/Source/MultiUserServer/Private/Widgets/Clients/PackageTransmission/Filter/PackageTransmissionFilter_FrontendRoot.h
@@ -0,0 +1,54 @@
+// Copyright Epic Games, Inc. All Rights Reserved.
+
+#pragma once
+
+#include "CoreMinimal.h"
+#include "PackageTransmissionFilter.h"
+
+class FEndpointToUserNameCache;
+class SWidget;
+
+namespace UE::MultiUserServer
+{
+	class FFrontendPackageTransmissionFilter;
+	class FFrontendPackageTransmissionFilter_TextSearch;
+	class FPackageTransmissionEntryTokenizer;
+
+	/** A filter that contains multiple UI filters */
+	class FPackageTransmissionFilter_FrontendRoot : public FPackageTransmissionFilter, public TSharedFromThis<FPackageTransmissionFilter_FrontendRoot>
+	{
+	public:
+
+		FPackageTransmissionFilter_FrontendRoot(
+			TSharedRef<FPackageTransmissionEntryTokenizer> Tokenizer,
+			TArray<TSharedRef<FFrontendPackageTransmissionFilter>> InCustomFilters,
+			const TArray<TSharedRef<FPackageTransmissionFilter>>& NonVisualFilters = {}
+			);
+
+		/** Builds the widget view for all contained filters */
+		TSharedRef<SWidget> BuildFilterWidgets() const;
+		
+		//~ Begin IFilter Interface
+		virtual bool PassesFilter(const FPackageTransmissionEntry& InItem) const override;
+		//~ End IFilter Interface
+
+		FORCEINLINE const TSharedRef<FFrontendPackageTransmissionFilter_TextSearch>& GetTextSearchFilter() const { return TextSearchFilter; }
+
+	private:
+
+		/** The text search filter. Also in FrontendFilters. Separate variable to build search bar in new line. */
+		TSharedRef<FFrontendPackageTransmissionFilter_TextSearch> TextSearchFilter;
+
+		/** AllFilters without special filters we have as properties above, such as TextSearchFilter. */
+		TArray<TSharedRef<FFrontendPackageTransmissionFilter>> FrontendFilters;
+		
+		/** Filters that are combined using logical AND. */
+		TArray<TSharedRef<FPackageTransmissionFilter>> AllFilters;
+
+		/** Builds the widgets that go under the text */
+		TSharedRef<SWidget> BuildCustomFilterListWidget() const;
+	};
+
+	/** Creates a filter for the global filter log window */
+	TSharedRef<FPackageTransmissionFilter_FrontendRoot> MakeFilter(TSharedRef<FPackageTransmissionEntryTokenizer> Tokenizer);
+}
diff --git a/Engine/Plugins/Developer/Concert/ConcertApp/MultiUserServer/Source/MultiUserServer/Private/Widgets/Clients/PackageTransmission/Model/FilteredPackageTransmissionModel.cpp b/Engine/Plugins/Developer/Concert/ConcertApp/MultiUserServer/Source/MultiUserServer/Private/Widgets/Clients/PackageTransmission/Model/FilteredPackageTransmissionModel.cpp
new file mode 100644
--- /dev/null
+++ b/Engine/Plugins/Developer/Concert/ConcertApp/MultiUserServer/Source/MultiUserServer/Private/Widgets/Clients/PackageTransmission/Model/FilteredPackageTransmissionModel.cpp
@@ -0,0 +1,97 @@
+// Copyright Epic Games, Inc. All Rights Reserved.
+
+#include "FilteredPackageTransmissionModel.h"
+
+#include "Widgets/Clients/PackageTransmission/Filter/PackageTransmissionFilter.h"
+
+namespace UE::MultiUserServer
+{
+	FFilteredPackageTransmissionModel::FFilteredPackageTransmissionModel(TSharedRef<IPackageTransmissionEntrySource> RealSource, TSharedRef<FPackageTransmissionFilter> Filter)
+		: RealSource(MoveTemp(RealSource))
+		, Filter(MoveTemp(Filter))
+	{
+		RealSource->OnPackageEntriesAdded().AddRaw(this, &FFilteredPackageTransmissionModel::OnSourceEntriesAdded);
+		RealSource->OnPackageEntriesModified().AddRaw(this, &FFilteredPackageTransmissionModel::OnSourceEntriesModified);
+
+		Filter->OnChanged().AddRaw(this, &FFilteredPackageTransmissionModel::ReapplyFilter, false);
+	}
+
+	FFilteredPackageTransmissionModel::~FFilteredPackageTransmissionModel()
+	{
+		RealSource->OnPackageEntriesAdded().RemoveAll(this);
+		RealSource->OnPackageEntriesModified().RemoveAll(this);
+
+		Filter->OnChanged().RemoveAll(this);
+	}
+
+	TSharedPtr<FPackageTransmissionEntry> FFilteredPackageTransmissionModel::GetEntryById(FPackageTransmissionId ID)
+	{
+		const TOptional<FPackageTransmissionId> Index = TranslatedSourceIndexToFilteredIndex(ID);
+		return Index
+			? FilteredEntries[*Index]
+			: TSharedPtr<FPackageTransmissionEntry>{};
+	}
+
+	void FFilteredPackageTransmissionModel::OnSourceEntriesAdded(uint32 NumAdded)
+	{
+		ApplyFilter(false, RealSource->GetEntries().Num() - NumAdded);
+	}
+
+	void FFilteredPackageTransmissionModel::OnSourceEntriesModified(const TSet<FPackageTransmissionId>& ModifiedEntries)
+	{
+		// Could be more efficient by only looking at ModifiedEntries... but usually the number of items is pretty small
+		ReapplyFilter(true);
+
+		TSet<FPackageTransmissionId> FilteredModified;
+		FilteredModified.Reserve(ModifiedEntries.Num());
+		for (FPackageTransmissionId Id : ModifiedEntries)
+		{
+			if (Filter->PassesFilter(*RealSource->GetEntryById(Id)))
+			{
+				FilteredModified.Add(Id);
+			}
+		}
+
+		if (FilteredModified.Num() > 0)
+		{
+			OnPackageEntriesModified().Broadcast(FilteredModified);
+		}
+	}
+
+	TOptional<FPackageTransmissionId> FFilteredPackageTransmissionModel::TranslatedSourceIndexToFilteredIndex(const FPackageTransmissionId OriginalId) const
+	{
+		const FPackageTransmissionId* ReboundIndex = RebindRealSourceToFilteredIndex.Find(OriginalId);
+		return ensure(ReboundIndex && FilteredEntries.IsValidIndex(*ReboundIndex))
+			? *ReboundIndex
+			: TOptional<FPackageTransmissionId>{};
+	}
+
+	void FFilteredPackageTransmissionModel::ReapplyFilter(bool bSilent)
+	{
+		FilteredEntries.Empty(FilteredEntries.Num());
+		RebindRealSourceToFilteredIndex.Empty(RebindRealSourceToFilteredIndex.Num());
+		ApplyFilter(bSilent);
+	}
+
+	void FFilteredPackageTransmissionModel::ApplyFilter(bool bSilent, int32 StartIndex)
+	{
+		uint32 NumAdded = 0;
+		
+		for (int32 i = StartIndex; i < RealSource->GetEntries().Num(); ++i)
+		{
+			const TSharedPtr<FPackageTransmissionEntry>& Entry = RealSource->GetEntries()[i];
+			if (Filter->PassesFilter(*Entry))
+			{
+				FPackageTransmissionId Index = FilteredEntries.Add(Entry);
+				RebindRealSourceToFilteredIndex.Add(Entry->TransmissionId, Index);
+				++NumAdded;
+			}
+		}
+
+		if (!bSilent && NumAdded > 0)
+		{
+			OnPackageEntriesAdded().Broadcast(NumAdded);
+		}
+	}
+}
+
diff --git a/Engine/Plugins/Developer/Concert/ConcertApp/MultiUserServer/Source/MultiUserServer/Private/Widgets/Clients/PackageTransmission/Model/FilteredPackageTransmissionModel.h b/Engine/Plugins/Developer/Concert/ConcertApp/MultiUserServer/Source/MultiUserServer/Private/Widgets/Clients/PackageTransmission/Model/FilteredPackageTransmissionModel.h
new file mode 100644
--- /dev/null
+++ b/Engine/Plugins/Developer/Concert/ConcertApp/MultiUserServer/Source/MultiUserServer/Private/Widgets/Clients/PackageTransmission/Model/FilteredPackageTransmissionModel.h
@@ -0,0 +1,47 @@
+// Copyright Epic Games, Inc. All Rights Reserved.
+
+#pragma once
+
+#include "CoreMinimal.h"
+#include "IPackageTransmissionEntrySource.h"
+#include "Widgets/Clients/PackageTransmission/Filter/PackageTransmissionFilter.h"
+
+namespace UE::MultiUserServer
+{
+	class FPackageTransmissionFilter;
+
+	class FFilteredPackageTransmissionModel : public IPackageTransmissionEntrySource
+	{
+	public:
+
+		FFilteredPackageTransmissionModel(TSharedRef<IPackageTransmissionEntrySource> RealSource, TSharedRef<FPackageTransmissionFilter> Filter);
+		virtual ~FFilteredPackageTransmissionModel() override;
+
+		//~ Begin IPackageTransmissionEntrySource Interface
+		virtual const TArray<TSharedPtr<FPackageTransmissionEntry>>& GetEntries() override { return FilteredEntries; }
+		virtual TSharedPtr<FPackageTransmissionEntry> GetEntryById(FPackageTransmissionId ID) override;
+		virtual FOnPackageEntriesAdded& OnPackageEntriesAdded() override { return OnPackageEntriesAddedDelegate; }
+		virtual FOnPackageEntriesModified& OnPackageEntriesModified() override { return OnPackageEntriesModifiedDelegate; }
+		//~ End IPackageTransmissionEntrySource Interface
+
+	private:
+
+		TSharedRef<IPackageTransmissionEntrySource> RealSource;
+		TSharedRef<FPackageTransmissionFilter> Filter;
+
+		FOnPackageEntriesAdded OnPackageEntriesAddedDelegate;
+		FOnPackageEntriesModified OnPackageEntriesModifiedDelegate;
+
+		TArray<TSharedPtr<FPackageTransmissionEntry>> FilteredEntries;
+		TMap<FPackageTransmissionId, FPackageTransmissionId> RebindRealSourceToFilteredIndex;
+
+		void OnSourceEntriesAdded(uint32 NumAdded);
+		void OnSourceEntriesModified(const TSet<FPackageTransmissionId>& ModifiedEntries);
+
+		TOptional<FPackageTransmissionId> TranslatedSourceIndexToFilteredIndex(const FPackageTransmissionId OriginalId) const;
+		
+		void ReapplyFilter(bool bSilent);
+		void ApplyFilter(bool bSilent = false, int32 StartIndex = 0);
+	};
+}
+
diff --git a/Engine/Plugins/Developer/Concert/ConcertApp/MultiUserServer/Source/MultiUserServer/Private/Widgets/Clients/PackageTransmission/Model/IPackageTransmissionEntrySource.h b/Engine/Plugins/Developer/Concert/ConcertApp/MultiUserServer/Source/MultiUserServer/Private/Widgets/Clients/PackageTransmission/Model/IPackageTransmissionEntrySource.h
new file mode 100644
--- /dev/null
+++ b/Engine/Plugins/Developer/Concert/ConcertApp/MultiUserServer/Source/MultiUserServer/Private/Widgets/Clients/PackageTransmission/Model/IPackageTransmissionEntrySource.h
@@ -0,0 +1,32 @@
+// Copyright Epic Games, Inc. All Rights Reserved.
+
+#pragma once
+
+#include "CoreMinimal.h"
+#include "PackageTransmissionEntry.h"
+#include "Templates/NonNullPointer.h"
+
+namespace UE::MultiUserServer
+{
+	/** Interfaces for an object that provides access to package transmissions - could be filtered or unfiltered. */
+	class IPackageTransmissionEntrySource
+	{
+	public:
+
+		virtual const TArray<TSharedPtr<FPackageTransmissionEntry>>& GetEntries() = 0;
+
+		/** @return The entry if valid - nullptr if not valid */
+		virtual TSharedPtr<FPackageTransmissionEntry> GetEntryById(FPackageTransmissionId ID) = 0;
+
+		/** Special version of OnPackageEntryArrayChanged that is called when entries are added. The last NumEntries number of elements in GetEntries are new. */
+		DECLARE_MULTICAST_DELEGATE_OneParam(FOnPackageEntriesAdded, uint32 /*NumEntries*/)
+		virtual FOnPackageEntriesAdded& OnPackageEntriesAdded() = 0;
+		
+		/** Called when the contents of an entry is changes. */
+		DECLARE_MULTICAST_DELEGATE_OneParam(FOnPackageEntriesModified, const TSet<FPackageTransmissionId>& /*ModifiedEntries*/);
+		virtual FOnPackageEntriesModified& OnPackageEntriesModified() = 0;
+
+		virtual ~IPackageTransmissionEntrySource() = default;
+	};
+}
+
diff --git a/Engine/Plugins/Developer/Concert/ConcertApp/MultiUserServer/Source/MultiUserServer/Private/Widgets/Clients/PackageTransmission/Model/PackageTransmissionEntry.h b/Engine/Plugins/Developer/Concert/ConcertApp/MultiUserServer/Source/MultiUserServer/Private/Widgets/Clients/PackageTransmission/Model/PackageTransmissionEntry.h
new file mode 100644
--- /dev/null
+++ b/Engine/Plugins/Developer/Concert/ConcertApp/MultiUserServer/Source/MultiUserServer/Private/Widgets/Clients/PackageTransmission/Model/PackageTransmissionEntry.h
@@ -0,0 +1,77 @@
+// Copyright Epic Games, Inc. All Rights Reserved.
+
+#pragma once
+
+#include "CoreMinimal.h"
+#include "ConcertWorkspaceData.h"
+
+namespace UE::MultiUserServer
+{
+	using FPackageTransmissionId = uint64;
+
+	enum class EPackageTransmissionState : uint8
+	{
+		/** We were told that the transmission has started and waiting for the destination to receive it. */
+		InTransmission,
+		/** The package was received and the data was not rejected; only the server rejects data. */
+		ReceivedAndAccepted,
+		/** The server rejected the update (bad request, client had no lock, multiple clients modified concurrently, etc.) */
+		Rejected
+	};
+
+	enum class EPackageSendDirection : uint8
+	{
+		ClientToServer,
+		ServerToClient
+	};
+	
+	/**
+	 * Describes the transmission of a single packet from one network point to another.
+	 *
+	 * Package transmission to the server happens in two stages:
+	 * 1. The client announces it is about to send a package - that's when this entry is created
+	 * 2. The server receives the package data - that's when MessageId is set
+	 *
+	 * Package transmission to the client happens in two stages:
+	 * 1. The server sends all package data to the client - that's when this entry is created
+	 * 2. The client sends an ACK - that's when MessageID is set
+	 */
+	struct FPackageTransmissionEntry
+	{
+		/** Unique ID created by our internal package transmission model - it's internally an array entry (implementation detail -  don't assume it) */
+		FPackageTransmissionId TransmissionId;
+
+		FDateTime LocalStartTime = FDateTime::Now(); 
+		EPackageTransmissionState TransmissionState = EPackageTransmissionState::InTransmission;
+		
+		EPackageSendDirection SendDirection;
+		
+		FConcertPackageInfo PackageInfo;
+		uint64 Revision;
+		
+		/** The origin endpoint announced that the package would have this size */
+		uint64 PackageNumBytes;
+		
+		/** The endpoint ID of the participating client. */
+		FGuid ClientEndpointId;
+
+		/**
+		 * The ID of the message that causes the package to be acknowledged.
+		 * Only set if TransmissionState != EPackageTransmissionState::InTransmission.
+		 *
+		 * In the case of sending to a client, it is the acknowledgement message.
+		 * In the case of receiving from client, it is the message transmitting the sync activity.
+		 */
+		TOptional<FGuid> MessageId;
+
+		FPackageTransmissionEntry(FPackageTransmissionId TransmissionId, EPackageSendDirection SendDirection, const FConcertPackageInfo& PackageInfo, uint64 Revision, uint64 PackageNumBytes, const FGuid& RemoteEndpointId)
+			: TransmissionId(TransmissionId)
+			, SendDirection(SendDirection)
+			, PackageInfo(PackageInfo)
+			, Revision(Revision)
+			, PackageNumBytes(PackageNumBytes)
+			, ClientEndpointId(RemoteEndpointId)
+		{}
+	};
+}
+
diff --git a/Engine/Plugins/Developer/Concert/ConcertApp/MultiUserServer/Source/MultiUserServer/Private/Widgets/Clients/PackageTransmission/Model/PackageTransmissionModel.cpp b/Engine/Plugins/Developer/Concert/ConcertApp/MultiUserServer/Source/MultiUserServer/Private/Widgets/Clients/PackageTransmission/Model/PackageTransmissionModel.cpp
new file mode 100644
--- /dev/null
+++ b/Engine/Plugins/Developer/Concert/ConcertApp/MultiUserServer/Source/MultiUserServer/Private/Widgets/Clients/PackageTransmission/Model/PackageTransmissionModel.cpp
@@ -0,0 +1,141 @@
+// Copyright Epic Games, Inc. All Rights Reserved.
+
+#include "PackageTransmissionModel.h"
+
+#include "ConcertLogGlobal.h"
+#include "ConcertPackageEvents.h"
+#include "ConcertSyncSessionDatabase.h"
+#include "ConcertUtil.h"
+#include "IConcertServer.h"
+#include "IConcertSyncServer.h"
+
+namespace UE::MultiUserServer
+{
+	FPackageTransmissionModel::FPackageTransmissionModel(TSharedRef<IConcertSyncServer> Server)
+		: Server(MoveTemp(Server))
+	{
+		using namespace ConcertSyncCore;
+		
+		ConcertPackageEvents::OnLocalBeginSendPackage().AddRaw(this, &FPackageTransmissionModel::OnLocalBeginSendPackage);
+		ConcertPackageEvents::OnLocalFinishSendPackage().AddRaw(this, &FPackageTransmissionModel::OnLocalFinishSendPackage);
+		ConcertPackageEvents::OnRemoteBeginSendPackage().AddRaw(this, &FPackageTransmissionModel::OnRemoteBeginSendPackage);
+		ConcertPackageEvents::OnRemoteFinishSendPackage().AddRaw(this, &FPackageTransmissionModel::OnRemoteFinishSendPackage);
+		ConcertPackageEvents::OnRejectRemoteSendPackage().AddRaw(this, &FPackageTransmissionModel::OnRejectRemoteSendPackage);
+
+		TickHandle = FTSTicker::GetCoreTicker().AddTicker(FTickerDelegate::CreateRaw(this, &FPackageTransmissionModel::Tick));
+	}
+
+	FPackageTransmissionModel::~FPackageTransmissionModel()
+	{
+		using namespace ConcertSyncCore;
+		
+		ConcertPackageEvents::OnLocalBeginSendPackage().RemoveAll(this);
+		ConcertPackageEvents::OnLocalFinishSendPackage().RemoveAll(this);
+		ConcertPackageEvents::OnRemoteBeginSendPackage().RemoveAll(this);
+		ConcertPackageEvents::OnRemoteFinishSendPackage().RemoveAll(this);
+		ConcertPackageEvents::OnRejectRemoteSendPackage().RemoveAll(this);
+		
+		FTSTicker::GetCoreTicker().RemoveTicker(TickHandle);
+	}
+
+	void FPackageTransmissionModel::OnLocalBeginSendPackage(const ConcertSyncCore::ConcertPackageEvents::FConcertBeginSendPackageParams& Args)
+	{
+		HandleBeginSend(EPackageSendDirection::ServerToClient, Args);
+	}
+
+	void FPackageTransmissionModel::OnLocalFinishSendPackage(const ConcertSyncCore::ConcertPackageEvents::FConcertFinishSendPackageParams& Args)
+	{
+		HandleFinishSend(EPackageTransmissionState::ReceivedAndAccepted, Args);
+	}
+
+	void FPackageTransmissionModel::OnRemoteBeginSendPackage(const ConcertSyncCore::ConcertPackageEvents::FConcertBeginSendPackageParams& Args)
+	{
+		HandleBeginSend(EPackageSendDirection::ClientToServer, Args);
+	}
+
+	void FPackageTransmissionModel::OnRemoteFinishSendPackage(const ConcertSyncCore::ConcertPackageEvents::FConcertFinishSendPackageParams& Args)
+	{
+		HandleFinishSend(EPackageTransmissionState::ReceivedAndAccepted, Args);
+	}
+
+	void FPackageTransmissionModel::OnRejectRemoteSendPackage(const ConcertSyncCore::ConcertPackageEvents::FConcertRejectSendPackageParams& Args)
+	{
+		HandleFinishSend(EPackageTransmissionState::Rejected, Args);
+	}
+
+	void FPackageTransmissionModel::HandleBeginSend(EPackageSendDirection SendDirection, const ConcertSyncCore::ConcertPackageEvents::FConcertBeginSendPackageParams& Args)
+	{
+		UE_LOG(LogConcert, Verbose, TEXT("OnLocalBeginSendPackage: %s"), *Args.TransmissionId.ToString(EGuidFormats::DigitsWithHyphens));
+		check(Args.PackageNumBytes > 0);
+
+		const FPackageTransmissionId ID = Entries.Num();
+		Entries.Add(MakeShared<FPackageTransmissionEntry>(ID, SendDirection, Args.PackageInfo, DeterminePackageRevision(SendDirection, Args), Args.PackageNumBytes, Args.RemoteEndpointId));
+		InTransmissionToIndex.Add(Args.TransmissionId, ID);
+		
+		++NumAddedSinceLastTick;
+	}
+
+	void FPackageTransmissionModel::HandleFinishSend(EPackageTransmissionState TransmissionState, const ConcertSyncCore::ConcertPackageEvents::FConcertFinishSendPackageParams& Args)
+	{
+		FPackageTransmissionId Key;
+		const bool bHadBeginSendData = InTransmissionToIndex.RemoveAndCopyValue(Args.TransmissionId, Key);
+		if (ensure(bHadBeginSendData) && ensure(Entries.IsValidIndex(Key)))
+		{
+			TSharedPtr<FPackageTransmissionEntry>& Entry = Entries[Key];
+			Entry->MessageId = Args.MessageId;
+			Entry->TransmissionState = TransmissionState;
+
+			ModifiedSinceLastTick.Add(Key);
+		}
+	}
+
+	uint64 FPackageTransmissionModel::DeterminePackageRevision(EPackageSendDirection SendDirection, const ConcertSyncCore::ConcertPackageEvents::FConcertBeginSendPackageParams& Args) const
+	{
+		// Show obviously wrong number
+		constexpr uint64 ErrorResult = static_cast<uint64>(-1);
+		
+		const TSharedPtr<IConcertServerSession> Session = ConcertUtil::GetLiveSessionClientConnectedTo(*Server->GetConcertServer(), Args.RemoteEndpointId);
+		if (!Session)
+		{
+			UE_LOG(LogConcert, Warning, TEXT("Processing package event: failed to get head revision because no session was found for client."));
+			return ErrorResult;
+		}
+		const TOptional<FConcertSyncSessionDatabaseNonNullPtr> Database = Server->GetLiveSessionDatabase(Session->GetId());
+		if (!Database)
+		{
+			UE_LOG(LogConcert, Warning, TEXT("Processing package event: failed to get head revision because session database could not be accessed."));
+			return ErrorResult;
+		}
+
+		int64 PackageRevision;
+		const bool bCouldGetHeadRevision = Database->GetPackageHeadRevision(Args.PackageInfo.PackageName, PackageRevision);
+		if (ensureMsgf(bCouldGetHeadRevision, TEXT("The activity should have been processed properly before this event was emitted")))
+		{
+			return SendDirection == EPackageSendDirection::ClientToServer
+				? PackageRevision + 1	// The moment the client sends to server, the revision is not created, yet. We're computing what the next revision will be.
+				: PackageRevision;
+		}
+
+		UE_LOG(LogConcert, Error, TEXT("Processing package event: failed to get head revision because the package was not saved in the database."));
+		return ErrorResult;
+	}
+
+	bool FPackageTransmissionModel::Tick(float)
+	{
+		// The rational of using Tick is to avoid excessive event calls which could cause filtering and possibly UI operations to rerun often
+		
+		if (NumAddedSinceLastTick > 0)
+		{
+			OnPackageEntriesAddedDelegate.Broadcast(NumAddedSinceLastTick);
+			NumAddedSinceLastTick = 0;
+		}
+
+		if (ModifiedSinceLastTick.Num() > 0)
+		{
+			OnPackageEntriesModifiedDelegate.Broadcast(ModifiedSinceLastTick);
+			ModifiedSinceLastTick.Empty(ModifiedSinceLastTick.Num());
+		}
+
+		return true;
+	}
+}
diff --git a/Engine/Plugins/Developer/Concert/ConcertApp/MultiUserServer/Source/MultiUserServer/Private/Widgets/Clients/PackageTransmission/Model/PackageTransmissionModel.h b/Engine/Plugins/Developer/Concert/ConcertApp/MultiUserServer/Source/MultiUserServer/Private/Widgets/Clients/PackageTransmission/Model/PackageTransmissionModel.h
new file mode 100644
--- /dev/null
+++ b/Engine/Plugins/Developer/Concert/ConcertApp/MultiUserServer/Source/MultiUserServer/Private/Widgets/Clients/PackageTransmission/Model/PackageTransmissionModel.h
@@ -0,0 +1,61 @@
+// Copyright Epic Games, Inc. All Rights Reserved.
+
+#pragma once
+
+#include "CoreMinimal.h"
+#include "Containers/Ticker.h"
+#include "IPackageTransmissionEntrySource.h"
+
+class IConcertSyncServer;
+
+namespace UE::ConcertSyncCore::ConcertPackageEvents
+{
+	struct FConcertBeginSendPackageParams;
+	struct FConcertFinishSendPackageParams;
+	struct FConcertRejectSendPackageParams;
+}
+
+namespace UE::MultiUserServer
+{
+	class FPackageTransmissionModel : public IPackageTransmissionEntrySource 
+	{
+	public:
+		
+		FPackageTransmissionModel(TSharedRef<IConcertSyncServer> Server);
+		virtual ~FPackageTransmissionModel() override;
+
+		//~ Begin IPackageTransmissionEntrySource Interface
+		virtual const TArray<TSharedPtr<FPackageTransmissionEntry>>& GetEntries() override { return Entries; }
+		virtual TSharedPtr<FPackageTransmissionEntry> GetEntryById(FPackageTransmissionId ID) override { return Entries.IsValidIndex(ID) ? Entries[ID] : TSharedPtr<FPackageTransmissionEntry>{}; }
+		virtual FOnPackageEntriesAdded& OnPackageEntriesAdded() override { return OnPackageEntriesAddedDelegate; }
+		virtual FOnPackageEntriesModified& OnPackageEntriesModified() override { return OnPackageEntriesModifiedDelegate; }
+		//~ End IPackageTransmissionEntrySource Interface
+
+	private:
+
+		TSharedRef<IConcertSyncServer> Server;
+		TArray<TSharedPtr<FPackageTransmissionEntry>> Entries;
+
+		FOnPackageEntriesAdded OnPackageEntriesAddedDelegate;
+		FOnPackageEntriesModified OnPackageEntriesModifiedDelegate;
+		
+		FTSTicker::FDelegateHandle TickHandle;
+		uint32 NumAddedSinceLastTick = 0;
+		TSet<FPackageTransmissionId> ModifiedSinceLastTick;
+		TMap<FGuid, FPackageTransmissionId> InTransmissionToIndex;
+		
+		void OnLocalBeginSendPackage(const ConcertSyncCore::ConcertPackageEvents::FConcertBeginSendPackageParams& Args);
+		void OnLocalFinishSendPackage(const ConcertSyncCore::ConcertPackageEvents::FConcertFinishSendPackageParams& Args);
+		void OnRemoteBeginSendPackage(const ConcertSyncCore::ConcertPackageEvents::FConcertBeginSendPackageParams& Args);
+		void OnRemoteFinishSendPackage(const ConcertSyncCore::ConcertPackageEvents::FConcertFinishSendPackageParams& Args);
+		void OnRejectRemoteSendPackage(const ConcertSyncCore::ConcertPackageEvents::FConcertRejectSendPackageParams& Args);
+
+		void HandleBeginSend(EPackageSendDirection SendDirection, const ConcertSyncCore::ConcertPackageEvents::FConcertBeginSendPackageParams& Args);
+		void HandleFinishSend(EPackageTransmissionState TransmissionState, const ConcertSyncCore::ConcertPackageEvents::FConcertFinishSendPackageParams& Args);
+
+		uint64 DeterminePackageRevision(EPackageSendDirection SendDirection, const ConcertSyncCore::ConcertPackageEvents::FConcertBeginSendPackageParams& Args) const;
+		
+		bool Tick(float);
+	};
+}
+
diff --git a/Engine/Plugins/Developer/Concert/ConcertApp/MultiUserServer/Source/MultiUserServer/Private/Widgets/Clients/PackageTransmission/PackageTransmissionTabController.cpp b/Engine/Plugins/Developer/Concert/ConcertApp/MultiUserServer/Source/MultiUserServer/Private/Widgets/Clients/PackageTransmission/PackageTransmissionTabController.cpp
new file mode 100644
--- /dev/null
+++ b/Engine/Plugins/Developer/Concert/ConcertApp/MultiUserServer/Source/MultiUserServer/Private/Widgets/Clients/PackageTransmission/PackageTransmissionTabController.cpp
@@ -0,0 +1,56 @@
+// Copyright Epic Games, Inc. All Rights Reserved.
+
+#include "PackageTransmissionTabController.h"
+
+#include "Model/PackageTransmissionModel.h"
+#include "SPackageTransmissionView.h"
+#include "Util/PackageTransmissionEntryTokenizer.h"
+
+#include "Framework/Docking/TabManager.h"
+#include "Widgets/Docking/SDockTab.h"
+
+#define LOCTEXT_NAMESPACE "UnrealMultiUserUI.FPackageTransmissionTabController"
+
+namespace UE::MultiUserServer
+{
+	FPackageTransmissionTabController::FPackageTransmissionTabController(
+			FName TabId,
+			TSharedRef<FTabManager> OwningTabManager,
+			TSharedRef<FPackageTransmissionModel> TransmissionModel,
+			TSharedRef<FEndpointToUserNameCache> EndpointToUserNameCache,
+			FCanScrollToLog CanScrollToLogDelegate,
+			FScrollToLog ScrollToLogDelegate
+			)
+		: TabId(TabId)
+		, OwningTabManager(MoveTemp(OwningTabManager))
+		, TransmissionModel(MoveTemp(TransmissionModel))
+		, EndpointToUserNameCache(MoveTemp(EndpointToUserNameCache))
+		, CanScrollToLogDelegate(MoveTemp(CanScrollToLogDelegate))
+		, ScrollToLogDelegate(MoveTemp(ScrollToLogDelegate))
+		, Tokenizer(MakeShared<FPackageTransmissionEntryTokenizer>(EndpointToUserNameCache))
+	{
+		OwningTabManager->RegisterTabSpawner(TabId, FOnSpawnTab::CreateRaw(this, &FPackageTransmissionTabController::SpawnTab))
+			// In the future we may create multiple FPackageTransmissionTabController and may want to change this name to depend on some parameter
+			.SetDisplayName(LOCTEXT("PackageTabLabel", "Packages"));
+	}
+
+	FPackageTransmissionTabController::~FPackageTransmissionTabController()
+	{
+		// Should not be needed because we're most likely destroyed at the same time as when the tab manager is destroyed but let's stay forward-compatible
+		OwningTabManager->UnregisterTabSpawner(TabId);
+	}
+
+	TSharedRef<SDockTab> FPackageTransmissionTabController::SpawnTab(const FSpawnTabArgs& SpawnTabArgs)
+	{
+		return SAssignNew(DockTab, SDockTab)
+			.Label(LOCTEXT("PackageTabLabel", "Packages"))
+			.TabRole(PanelTab)
+			[
+				SNew(SPackageTransmissionView, TransmissionModel, Tokenizer)
+				.CanScrollToLog(CanScrollToLogDelegate)
+				.ScrollToLog(ScrollToLogDelegate)
+			];
+	}
+}
+
+#undef LOCTEXT_NAMESPACE
\ No newline at end of file
diff --git a/Engine/Plugins/Developer/Concert/ConcertApp/MultiUserServer/Source/MultiUserServer/Private/Widgets/Clients/PackageTransmission/PackageTransmissionTabController.h b/Engine/Plugins/Developer/Concert/ConcertApp/MultiUserServer/Source/MultiUserServer/Private/Widgets/Clients/PackageTransmission/PackageTransmissionTabController.h
new file mode 100644
--- /dev/null
+++ b/Engine/Plugins/Developer/Concert/ConcertApp/MultiUserServer/Source/MultiUserServer/Private/Widgets/Clients/PackageTransmission/PackageTransmissionTabController.h
@@ -0,0 +1,50 @@
+// Copyright Epic Games, Inc. All Rights Reserved.
+
+#pragma once
+
+#include "CoreMinimal.h"
+#include "Widgets/Clients/Logging/LogScrollingDelegates.h"
+
+class FEndpointToUserNameCache;
+class FSpawnTabArgs;
+class FTabManager;
+class IConcertSyncServer;
+class SDockTab;
+
+namespace UE::MultiUserServer
+{
+	class FPackageTransmissionEntryTokenizer;
+	class FPackageTransmissionModel;
+
+	/** Manages a tab that displays a package transmission model. The model */
+	class FPackageTransmissionTabController
+	{
+	public:
+		
+		FPackageTransmissionTabController(
+			FName TabId,
+			TSharedRef<FTabManager> OwningTabManager,
+			TSharedRef<FPackageTransmissionModel> TransmissionModel,
+			TSharedRef<FEndpointToUserNameCache> EndpointToUserNameCache,
+			FCanScrollToLog CanScrollToLogDelegate,
+			FScrollToLog ScrollToLogDelegate
+			);
+		~FPackageTransmissionTabController();
+
+	private:
+		
+		FName TabId;
+		TSharedRef<FTabManager> OwningTabManager;
+		TSharedRef<FPackageTransmissionModel> TransmissionModel;
+		TSharedRef<FEndpointToUserNameCache> EndpointToUserNameCache;
+		
+		FCanScrollToLog CanScrollToLogDelegate;
+		FScrollToLog ScrollToLogDelegate;
+		
+		TSharedRef<FPackageTransmissionEntryTokenizer> Tokenizer;
+		TSharedPtr<SDockTab> DockTab;
+
+		TSharedRef<SDockTab> SpawnTab(const FSpawnTabArgs& SpawnTabArgs);
+	};
+}
+
diff --git a/Engine/Plugins/Developer/Concert/ConcertApp/MultiUserServer/Source/MultiUserServer/Private/Widgets/Clients/PackageTransmission/SPackageTransmissionTable.cpp b/Engine/Plugins/Developer/Concert/ConcertApp/MultiUserServer/Source/MultiUserServer/Private/Widgets/Clients/PackageTransmission/SPackageTransmissionTable.cpp
new file mode 100644
--- /dev/null
+++ b/Engine/Plugins/Developer/Concert/ConcertApp/MultiUserServer/Source/MultiUserServer/Private/Widgets/Clients/PackageTransmission/SPackageTransmissionTable.cpp
@@ -0,0 +1,202 @@
+// Copyright Epic Games, Inc. All Rights Reserved.
+
+#include "SPackageTransmissionTable.h"
+
+#include "ConcertHeaderRowUtils.h"
+#include "Model/IPackageTransmissionEntrySource.h"
+#include "Model/PackageTransmissionEntry.h"
+#include "Settings/MultiUserServerColumnVisibilitySettings.h"
+#include "SPackageTransmissionTableFooter.h"
+#include "SPackageTransmissionTableRow.h"
+
+#include "Framework/MultiBox/MultiBoxBuilder.h"
+#include "Widgets/SBoxPanel.h"
+#include "Widgets/Views/SHeaderRow.h"
+#include "Widgets/Views/SListView.h"
+
+#define LOCTEXT_NAMESPACE "UnrealMultiUserUI.SPackageTransmissionTable"
+
+namespace UE::MultiUserServer
+{
+	void SPackageTransmissionTable::Construct(const FArguments& InArgs, TSharedRef<IPackageTransmissionEntrySource> InPackageEntrySource, TSharedRef<FPackageTransmissionEntryTokenizer> InTokenizer)
+	{
+		PackageEntrySource = MoveTemp(InPackageEntrySource);
+		Tokenizer = MoveTemp(InTokenizer);
+
+		HighlightText = InArgs._HighlightText;
+		CanScrollToLogDelegate = InArgs._CanScrollToLog;
+		ScrollToLogDelegate = InArgs._ScrollToLog;
+		
+		ChildSlot
+		[
+			SNew(SVerticalBox)
+
+			+SVerticalBox::Slot()
+			.FillHeight(1.f)
+			.Padding(0, 5, 0, 0)
+			[
+				CreateTableView()
+			]
+
+			+SVerticalBox::Slot()
+			.AutoHeight()
+			[
+				SNew(SPackageTransmissionTableFooter, PackageEntrySource.ToSharedRef())
+				.ExtendViewOptions(this, &SPackageTransmissionTable::ExtendViewOptions)
+				.TotalUnfilteredNum(InArgs._TotalUnfilteredNum)
+			]
+		];
+
+		UMultiUserServerColumnVisibilitySettings::GetSettings()->OnOnPackageTransmissionColumnVisibilityChanged().AddSP(this, &SPackageTransmissionTable::OnColumnVisibilitySettingsChanged);
+		ConcertSharedSlate::RestoreColumnVisibilityState(HeaderRow.ToSharedRef(), UMultiUserServerColumnVisibilitySettings::GetSettings()->GetPackageTransmissionColumnVisibility());
+
+		PackageEntrySource->OnPackageEntriesModified().AddSP(this, &SPackageTransmissionTable::OnPackageEntriesModified);
+		PackageEntrySource->OnPackageEntriesAdded().AddSP(this, &SPackageTransmissionTable::OnPackageArrayChanged);
+	}
+
+	SPackageTransmissionTable::~SPackageTransmissionTable()
+	{
+		UMultiUserServerColumnVisibilitySettings::GetSettings()->OnOnPackageTransmissionColumnVisibilityChanged().RemoveAll(this);
+		PackageEntrySource->OnPackageEntriesModified().RemoveAll(this);
+		PackageEntrySource->OnPackageEntriesAdded().RemoveAll(this);
+	}
+
+	TSharedRef<SWidget> SPackageTransmissionTable::CreateTableView()
+	{
+		return SAssignNew(TableView, SListView<TSharedPtr<FPackageTransmissionEntry>>)
+			.ListItemsSource(&PackageEntrySource->GetEntries())
+			.OnGenerateRow(this, &SPackageTransmissionTable::OnGenerateActivityRowWidget)
+			.SelectionMode(ESelectionMode::Multi)
+			.HeaderRow(CreateHeaderRow());
+	}
+
+	TSharedRef<SHeaderRow> SPackageTransmissionTable::CreateHeaderRow()
+	{
+		HeaderRow = SNew(SHeaderRow)
+		.OnHiddenColumnsListChanged_Lambda([this]()
+		{
+			if (!bIsUpdatingColumnVisibility)
+			{
+				UMultiUserServerColumnVisibilitySettings::GetSettings()->SetPackageTransmissionColumnVisibility(
+				UE::ConcertSharedSlate::SnapshotColumnVisibilityState(HeaderRow.ToSharedRef())
+				);
+			}
+		});
+
+		const TSet<FName> CannotHideColumns = { SPackageTransmissionTableRow::TransmissionStateColumn };
+		for (FName ColumnName : SPackageTransmissionTableRow::AllColumns)
+		{
+			const bool bCannotHide = CannotHideColumns.Contains(ColumnName);
+			SHeaderRow::FColumn::FArguments Args = SHeaderRow::FColumn::FArguments()
+				.ColumnId(ColumnName)
+				.HAlignCell(HAlign_Center)
+				.DefaultLabel(SPackageTransmissionTableRow::ColumnsDisplayText[ColumnName]);
+
+			if (bCannotHide)
+			{
+				Args.ShouldGenerateWidget(bCannotHide);
+			}
+			else
+			{
+				Args.OnGetMenuContent_Lambda([this]()
+				{
+					return ConcertSharedSlate::MakeHideColumnContextMenu(HeaderRow.ToSharedRef(), SPackageTransmissionTableRow::TimeColumn);
+				});
+			}
+			
+			HeaderRow->AddColumn(Args);
+		}
+
+		TGuardValue<bool> DoNotSave(bIsUpdatingColumnVisibility, true);
+		RestoreDefaultColumnVisibilities();
+
+		return HeaderRow.ToSharedRef();
+	}
+
+	TSharedRef<ITableRow> SPackageTransmissionTable::OnGenerateActivityRowWidget(TSharedPtr<FPackageTransmissionEntry> Item, const TSharedRef<STableViewBase>& OwnerTable) const
+	{
+		return SNew(SPackageTransmissionTableRow, Item, OwnerTable, Tokenizer.ToSharedRef())
+			.HighlightText(HighlightText)
+			.CanScrollToLog(CanScrollToLogDelegate)
+			.ScrollToLog(ScrollToLogDelegate);
+	}
+
+	void SPackageTransmissionTable::OnColumnVisibilitySettingsChanged(const FColumnVisibilitySnapshot& ColumnSnapshot)
+	{
+		TGuardValue<bool> GuardValue(bIsUpdatingColumnVisibility, true);
+		ConcertSharedSlate::RestoreColumnVisibilityState(HeaderRow.ToSharedRef(), ColumnSnapshot);
+	}
+
+	void SPackageTransmissionTable::ExtendViewOptions(FMenuBuilder& MenuBuilder)
+	{
+		MenuBuilder.AddMenuEntry(
+			LOCTEXT("SelectAll", "Show all"),
+			FText::GetEmpty(),
+			FSlateIcon(),
+			FUIAction(
+				FExecuteAction::CreateLambda([this]()
+				{
+					for(FName ColumnId : SPackageTransmissionTableRow::AllColumns)
+					{
+						HeaderRow->SetShowGeneratedColumn(ColumnId, true);
+					}
+				}),
+				FCanExecuteAction::CreateLambda([] { return true; })),
+			NAME_None,
+			EUserInterfaceActionType::Button
+		);
+		MenuBuilder.AddMenuEntry(
+			LOCTEXT("HideAll", "Hide all"),
+			FText::GetEmpty(),
+			FSlateIcon(),
+			FUIAction(
+				FExecuteAction::CreateLambda([this]()
+				{
+					for(FName ColumnId : SPackageTransmissionTableRow::AllColumns)
+					{
+						HeaderRow->SetShowGeneratedColumn(ColumnId, false);
+					}
+				}),
+				FCanExecuteAction::CreateLambda([] { return true; })),
+			NAME_None,
+			EUserInterfaceActionType::Button
+		);
+		MenuBuilder.AddMenuEntry(
+				LOCTEXT("RestoreDefaultColumnVisibility", "Restore columns visibility"),
+				FText::GetEmpty(),
+				FSlateIcon(),
+				FUIAction(
+					FExecuteAction::CreateSP(this, &SPackageTransmissionTable::RestoreDefaultColumnVisibilities),
+					FCanExecuteAction::CreateLambda([] { return true; })),
+				NAME_None,
+				EUserInterfaceActionType::Button
+			);
+		MenuBuilder.AddSeparator();
+		ConcertSharedSlate::AddEntriesForShowingHiddenRows(HeaderRow.ToSharedRef(), MenuBuilder);
+	}
+
+	void SPackageTransmissionTable::RestoreDefaultColumnVisibilities()
+	{
+		const TSet<FName> HiddenByDefault = {
+			SPackageTransmissionTableRow::OriginColumn,
+			SPackageTransmissionTableRow::DestinationColumn,
+			SPackageTransmissionTableRow::PackagePathColumn
+		};
+		for (FName ColumnName : SPackageTransmissionTableRow::AllColumns)
+		{
+			HeaderRow->SetShowGeneratedColumn(ColumnName, !HiddenByDefault.Contains(ColumnName));
+		}
+	}
+
+	void SPackageTransmissionTable::OnPackageEntriesModified(const TSet<FPackageTransmissionId>& Set) const
+	{
+		TableView->RequestListRefresh();
+	}
+
+	void SPackageTransmissionTable::OnPackageArrayChanged(uint32 NumAdded) const
+	{
+		TableView->RequestListRefresh();
+	}
+}
+
+#undef LOCTEXT_NAMESPACE
\ No newline at end of file
diff --git a/Engine/Plugins/Developer/Concert/ConcertApp/MultiUserServer/Source/MultiUserServer/Private/Widgets/Clients/PackageTransmission/SPackageTransmissionTable.h b/Engine/Plugins/Developer/Concert/ConcertApp/MultiUserServer/Source/MultiUserServer/Private/Widgets/Clients/PackageTransmission/SPackageTransmissionTable.h
new file mode 100644
--- /dev/null
+++ b/Engine/Plugins/Developer/Concert/ConcertApp/MultiUserServer/Source/MultiUserServer/Private/Widgets/Clients/PackageTransmission/SPackageTransmissionTable.h
@@ -0,0 +1,66 @@
+// Copyright Epic Games, Inc. All Rights Reserved.
+
+#pragma once
+
+#include "CoreMinimal.h"
+#include "Model/PackageTransmissionEntry.h"
+#include "Widgets/DeclarativeSyntaxSupport.h"
+#include "Widgets/SCompoundWidget.h"
+#include "Widgets/Clients/Logging/LogScrollingDelegates.h"
+
+class FMenuBuilder;
+class ITableRow;
+class SHeaderRow;
+class STableViewBase;
+template<typename T> class SListView;
+struct FColumnVisibilitySnapshot;
+
+namespace UE::MultiUserServer
+{
+	class FPackageTransmissionEntryTokenizer;
+	class IPackageTransmissionEntrySource;
+	struct FPackageTransmissionEntry;
+
+	/** Displays sent and received packages in a table */
+	class SPackageTransmissionTable : public SCompoundWidget
+	{
+	public:
+
+		SLATE_BEGIN_ARGS(SPackageTransmissionTable)
+		{}
+			SLATE_ATTRIBUTE(uint32, TotalUnfilteredNum)
+			SLATE_ATTRIBUTE(FText, HighlightText)
+			SLATE_EVENT(FCanScrollToLog, CanScrollToLog)
+			SLATE_EVENT(FScrollToLog, ScrollToLog)
+		SLATE_END_ARGS()
+
+		void Construct(const FArguments& InArgs, TSharedRef<IPackageTransmissionEntrySource> InPackageEntrySource, TSharedRef<FPackageTransmissionEntryTokenizer> InTokenizer);
+		virtual ~SPackageTransmissionTable() override;
+		
+	private:
+
+		TSharedPtr<IPackageTransmissionEntrySource> PackageEntrySource;
+		TSharedPtr<FPackageTransmissionEntryTokenizer> Tokenizer;
+
+		TAttribute<FText> HighlightText;
+		FCanScrollToLog CanScrollToLogDelegate;
+		FScrollToLog ScrollToLogDelegate;
+		
+		TSharedPtr<SListView<TSharedPtr<FPackageTransmissionEntry>>> TableView;
+		TSharedPtr<SHeaderRow> HeaderRow;
+		/** Whether we currently loading the column visibility - prevents infinite event recursion */
+		bool bIsUpdatingColumnVisibility = false;
+
+		// Table view creation
+		TSharedRef<SWidget> CreateTableView();
+		TSharedRef<SHeaderRow> CreateHeaderRow();
+		TSharedRef<ITableRow> OnGenerateActivityRowWidget(TSharedPtr<FPackageTransmissionEntry> Item, const TSharedRef<STableViewBase>& OwnerTable) const;
+		
+		void OnColumnVisibilitySettingsChanged(const FColumnVisibilitySnapshot& ColumnSnapshot);
+		void ExtendViewOptions(FMenuBuilder& MenuBuilder);
+		void RestoreDefaultColumnVisibilities(); 
+		
+		void OnPackageEntriesModified(const TSet<FPackageTransmissionId>& Set) const;
+		void OnPackageArrayChanged(uint32 NumAdded) const;
+	};
+}
diff --git a/Engine/Plugins/Developer/Concert/ConcertApp/MultiUserServer/Source/MultiUserServer/Private/Widgets/Clients/PackageTransmission/SPackageTransmissionTableFooter.cpp b/Engine/Plugins/Developer/Concert/ConcertApp/MultiUserServer/Source/MultiUserServer/Private/Widgets/Clients/PackageTransmission/SPackageTransmissionTableFooter.cpp
new file mode 100644
--- /dev/null
+++ b/Engine/Plugins/Developer/Concert/ConcertApp/MultiUserServer/Source/MultiUserServer/Private/Widgets/Clients/PackageTransmission/SPackageTransmissionTableFooter.cpp
@@ -0,0 +1,91 @@
+// Copyright Epic Games, Inc. All Rights Reserved.
+
+#include "SPackageTransmissionTableFooter.h"
+
+#include "Model/IPackageTransmissionEntrySource.h"
+
+#include "Framework/MultiBox/MultiBoxBuilder.h"
+#include "Styling/AppStyle.h"
+#include "Widgets/Images/SImage.h"
+#include "Widgets/Input/SComboButton.h"
+#include "Widgets/Layout/SSpacer.h"
+#include "Widgets/SBoxPanel.h"
+#include "Widgets/Text/STextBlock.h"
+
+#define LOCTEXT_NAMESPACE "UnrealMultiUserUI.SConcertTransportLogRow"
+
+namespace UE::MultiUserServer
+{
+	void SPackageTransmissionTableFooter::Construct(const FArguments& InArgs, TSharedRef<IPackageTransmissionEntrySource> InPackageEntrySource)
+	{
+		PackageEntrySource = MoveTemp(InPackageEntrySource);
+		TotalUnfilteredNum = InArgs._TotalUnfilteredNum;
+		
+		if (!ensure(TotalUnfilteredNum.IsBound()))
+		{
+			TotalUnfilteredNum = TAttribute<uint32>::CreateLambda([this](){ return PackageEntrySource->GetEntries().Num(); });
+		}
+		
+		ChildSlot
+		[
+			SNew(SHorizontalBox)
+
+			// Logs per page text
+			+SHorizontalBox::Slot()
+			.AutoWidth()
+			.VAlign(VAlign_Center)
+			[
+				SNew(STextBlock)
+				.Text_Lambda([this]()
+				{
+					return FText::Format(LOCTEXT("DisplayPackagesFmt", "Displaying {0} of {1} packages"), PackageEntrySource->GetEntries().Num(), TotalUnfilteredNum.Get());
+				})
+			]
+
+			// Gap Filler
+			+SHorizontalBox::Slot()
+			.FillWidth(1.0)
+			[
+				SNew(SSpacer)
+			]
+
+			+SHorizontalBox::Slot()
+			.AutoWidth()
+			.VAlign(VAlign_Center)
+			[
+				SNew(SComboButton)
+					.OnGetMenuContent(this, &SPackageTransmissionTableFooter::MakeViewOptionsMenuWidget, InArgs._ExtendViewOptions)
+					.HasDownArrow(true)
+					.ButtonContent()
+					[
+						SNew(SHorizontalBox)
+
+						+SHorizontalBox::Slot()
+						.AutoWidth()
+						.VAlign(VAlign_Center)
+						[
+							SNew(SImage).Image(FAppStyle::Get().GetBrush("Icons.Visible")) // The eye ball image.
+						]
+
+						+SHorizontalBox::Slot()
+						.AutoWidth()
+						.Padding(2, 0, 0, 0)
+						.VAlign(VAlign_Center)
+						[
+							SNew(STextBlock).Text(LOCTEXT("ViewOptions", "View Options"))
+						]
+					]
+			]
+		];
+	}
+
+	TSharedRef<SWidget> SPackageTransmissionTableFooter::MakeViewOptionsMenuWidget(FExtendContextMenu ExtendViewOptions)
+	{
+		FMenuBuilder MenuBuilder(/*bInShouldCloseWindowAfterMenuSelection=*/true, nullptr);
+
+		ExtendViewOptions.Execute(MenuBuilder);
+		return MenuBuilder.MakeWidget();
+	}
+}
+
+#undef LOCTEXT_NAMESPACE
\ No newline at end of file
diff --git a/Engine/Plugins/Developer/Concert/ConcertApp/MultiUserServer/Source/MultiUserServer/Private/Widgets/Clients/PackageTransmission/SPackageTransmissionTableFooter.h b/Engine/Plugins/Developer/Concert/ConcertApp/MultiUserServer/Source/MultiUserServer/Private/Widgets/Clients/PackageTransmission/SPackageTransmissionTableFooter.h
new file mode 100644
--- /dev/null
+++ b/Engine/Plugins/Developer/Concert/ConcertApp/MultiUserServer/Source/MultiUserServer/Private/Widgets/Clients/PackageTransmission/SPackageTransmissionTableFooter.h
@@ -0,0 +1,38 @@
+// Copyright Epic Games, Inc. All Rights Reserved.
+
+#pragma once
+
+#include "CoreMinimal.h"
+#include "Widgets/DeclarativeSyntaxSupport.h"
+#include "Widgets/SCompoundWidget.h"
+
+class FMenuBuilder;
+class IConcertSyncServer;
+
+namespace UE::MultiUserServer
+{
+	class IPackageTransmissionEntrySource;
+
+	/** Contains all UI for displaying package transmission, including filters, the table, and view options. */
+	class SPackageTransmissionTableFooter : public SCompoundWidget
+	{
+	public:
+
+		DECLARE_DELEGATE_OneParam(FExtendContextMenu, FMenuBuilder&)
+
+		SLATE_BEGIN_ARGS(SPackageTransmissionTableFooter)
+		{}
+			SLATE_EVENT(FExtendContextMenu, ExtendViewOptions)
+			SLATE_ATTRIBUTE(uint32, TotalUnfilteredNum)
+		SLATE_END_ARGS()
+
+		void Construct(const FArguments& InArgs, TSharedRef<IPackageTransmissionEntrySource> InPackageEntrySource);
+
+	private:
+		
+		TSharedPtr<IPackageTransmissionEntrySource> PackageEntrySource;
+		TAttribute<uint32> TotalUnfilteredNum;
+		
+		TSharedRef<SWidget> MakeViewOptionsMenuWidget(FExtendContextMenu ExtendViewOptions);
+	};
+}
diff --git a/Engine/Plugins/Developer/Concert/ConcertApp/MultiUserServer/Source/MultiUserServer/Private/Widgets/Clients/PackageTransmission/SPackageTransmissionTableRow.cpp b/Engine/Plugins/Developer/Concert/ConcertApp/MultiUserServer/Source/MultiUserServer/Private/Widgets/Clients/PackageTransmission/SPackageTransmissionTableRow.cpp
new file mode 100644
--- /dev/null
+++ b/Engine/Plugins/Developer/Concert/ConcertApp/MultiUserServer/Source/MultiUserServer/Private/Widgets/Clients/PackageTransmission/SPackageTransmissionTableRow.cpp
@@ -0,0 +1,193 @@
+// Copyright Epic Games, Inc. All Rights Reserved.
+
+#include "SPackageTransmissionTableRow.h"
+
+#include "ConcertServerStyle.h"
+#include "Model/PackageTransmissionEntry.h"
+#include "Util/PackageTransmissionEntryTokenizer.h"
+
+#include "Widgets/SOverlay.h"
+#include "Widgets/Layout/SScaleBox.h"
+#include "Widgets/Images/SImage.h"
+#include "Widgets/Images/SThrobber.h"
+#include "Widgets/Input/SButton.h"
+#include "Widgets/Text/STextBlock.h"
+#include "Widgets/Views/STableRow.h"
+
+#define LOCTEXT_NAMESPACE "UnrealMultiUserUI.SPackageTransmissionTableRow"
+
+namespace UE::MultiUserServer
+{
+	const FName SPackageTransmissionTableRow::TimeColumn("TimeColumn");
+	const FName SPackageTransmissionTableRow::OriginColumn("OriginColumn");
+	const FName SPackageTransmissionTableRow::DestinationColumn("DestinationColumn");
+	const FName SPackageTransmissionTableRow::SizeColumn("SizeColumn");
+	const FName SPackageTransmissionTableRow::RevisionColumn("VersionColumn");
+	const FName SPackageTransmissionTableRow::PackagePathColumn("PackagePathColumn");
+	const FName SPackageTransmissionTableRow::PackageNameColumn("PackageNameColumn");
+	const FName SPackageTransmissionTableRow::TransmissionStateColumn("TransmissionStateColumn");
+	
+	const TArray<FName> SPackageTransmissionTableRow::AllColumns({ TimeColumn, OriginColumn, DestinationColumn, SizeColumn, RevisionColumn, PackagePathColumn, PackageNameColumn, TransmissionStateColumn });
+	const TMap<FName, FText> SPackageTransmissionTableRow::ColumnsDisplayText({
+		{ TimeColumn, LOCTEXT("TimeColumn", "Time") },
+		{ OriginColumn, LOCTEXT("OriginColumn", "Origin") },
+		{ DestinationColumn, LOCTEXT("DestinationColumn", "Destination") },
+		{ SizeColumn, LOCTEXT("SizeColumn", "Size") },
+		{ RevisionColumn, LOCTEXT("RevisionColumn", "Revision") },
+		{ PackagePathColumn, LOCTEXT("PackagePathColumn", "Package Path") },
+		{ PackageNameColumn, LOCTEXT("PackageNameColumn", "Package Name") },
+		{ TransmissionStateColumn, LOCTEXT("TransmissionStateColumn", "State") }
+	});
+	
+	void SPackageTransmissionTableRow::Construct(const FArguments& InArgs, TSharedPtr<FPackageTransmissionEntry> InPackageEntry, const TSharedRef<STableViewBase>& InOwnerTableView, TSharedRef<FPackageTransmissionEntryTokenizer> InTokenizer)
+	{
+		PackageEntry = MoveTemp(InPackageEntry);
+		Tokenizer = MoveTemp(InTokenizer);
+
+		HighlightText = InArgs._HighlightText;
+		CanScrollToLogDelegate = InArgs._CanScrollToLog;
+		ScrollToLogDelegate = InArgs._ScrollToLog;
+		
+		SMultiColumnTableRow<TSharedPtr<FPackageTransmissionEntry>>::Construct(FSuperRowType::FArguments(), InOwnerTableView);
+	}
+
+	TSharedRef<SWidget> SPackageTransmissionTableRow::GenerateWidgetForColumn(const FName& ColumnName)
+	{
+		if (ColumnName == TransmissionStateColumn)
+		{
+			return CreateTransmissionStateColumn();
+		}
+
+		return SNew(STextBlock)
+			.HighlightText(HighlightText)
+			.Text_Lambda([this, ColumnName]()
+			{
+				using FColumnWidgetFactoryFunc = FString(FPackageTransmissionEntryTokenizer::*)(const FPackageTransmissionEntry&) const;
+				const TMap<FName, FColumnWidgetFactoryFunc> OverrideFactories = {
+					{ TimeColumn, &FPackageTransmissionEntryTokenizer::TokenizeTime },
+					{ OriginColumn, &FPackageTransmissionEntryTokenizer::TokenizeOrigin },
+					{ DestinationColumn, &FPackageTransmissionEntryTokenizer::TokenizeDestination },
+					{ SizeColumn, &FPackageTransmissionEntryTokenizer::TokenizeSize },
+					{ RevisionColumn, &FPackageTransmissionEntryTokenizer::TokenizeRevision },
+					{ PackagePathColumn, &FPackageTransmissionEntryTokenizer::TokenizePackagePath },
+					{ PackageNameColumn, &FPackageTransmissionEntryTokenizer::TokenizePackageName }
+				};
+				
+				return FText::FromString(Invoke(OverrideFactories[ColumnName], *Tokenizer, *PackageEntry));
+			});
+	}
+	
+
+	TSharedRef<SWidget> SPackageTransmissionTableRow::CreateTransmissionStateColumn() const
+	{
+		return SNew(SOverlay)
+
+			+SOverlay::Slot()
+			.HAlign(HAlign_Center)
+			.VAlign(VAlign_Center)
+			[
+				CreateTransmissionInProgressWidget()
+			]
+
+			+SOverlay::Slot()
+			.HAlign(HAlign_Center)
+			.VAlign(VAlign_Center)
+			[
+				CreateTransmissionSuccessWidget()
+			]
+
+			+SOverlay::Slot()
+			.HAlign(HAlign_Center)
+			.VAlign(VAlign_Center)
+			[
+				CreateTransmissionFailureWidget()
+			];
+	}
+
+	TSharedRef<SWidget> SPackageTransmissionTableRow::CreateTransmissionInProgressWidget() const
+	{
+		return SNew(SScaleBox)
+			.Stretch(EStretch::ScaleToFit)
+			[
+				SNew(SThrobber)
+				.Animate(SThrobber::VerticalAndOpacity)
+				.Visibility_Lambda([this]()
+				{
+					return PackageEntry->TransmissionState == EPackageTransmissionState::InTransmission
+						? EVisibility::Visible
+						: EVisibility::Hidden;
+				})
+				.ToolTipText(LOCTEXT("Intransmission", "Transmitting..."))
+			];
+	}
+
+	TSharedRef<SWidget> SPackageTransmissionTableRow::CreateTransmissionSuccessWidget() const
+	{
+		return SharedCreateTransmissionCompletedWidget(EPackageTransmissionState::ReceivedAndAccepted, FConcertServerStyle::Get().GetBrush("Concert.PackageTransmission.Success"));
+	}
+
+	TSharedRef<SWidget> SPackageTransmissionTableRow::CreateTransmissionFailureWidget() const
+	{
+		return SharedCreateTransmissionCompletedWidget(EPackageTransmissionState::Rejected, FConcertServerStyle::Get().GetBrush("Concert.PackageTransmission.Failure"));
+	}
+
+	TSharedRef<SWidget> SPackageTransmissionTableRow::SharedCreateTransmissionCompletedWidget(EPackageTransmissionState AllowedTransmissionState, const FSlateBrush* ImageBrush) const
+	{
+		return SNew(SButton)
+			.ContentPadding(0.f)
+			.IsEnabled_Lambda([this]()
+			{
+				return CanScrollToLog();
+			})
+			.OnPressed_Lambda([this]()
+			{
+				ScrollToLog();
+			})
+			.Visibility_Lambda([this, AllowedTransmissionState]()
+			{
+				return PackageEntry->TransmissionState == AllowedTransmissionState
+					? EVisibility::Visible
+					: EVisibility::Hidden;
+			})
+			.ToolTipText_Lambda([this, AllowedTransmissionState]()
+			{
+				const FText BaseTooltip = PackageEntry->TransmissionState == AllowedTransmissionState
+					? LOCTEXT("PackageTransmission.Accepted", "Package received.")
+					: LOCTEXT("PackageTransmission.Rejected", "Package was rejected by the server.");
+
+				FText ErrorMessage;
+				const bool bCanScroll = CanScrollToLog(ErrorMessage);
+				return bCanScroll
+					? BaseTooltip
+					: FText::Format(LOCTEXT("PackageTransmissionFmt", "Cannot scroll to log: {0}\n{1}"), ErrorMessage, BaseTooltip);
+			})
+			[
+				SNew(SImage)
+				.DesiredSizeOverride(FVector2D(16.f, 16.f))
+				.Image(ImageBrush)
+			];
+	}
+	
+	bool SPackageTransmissionTableRow::CanScrollToLog(FText& ErrorMessage) const
+	{
+		return PackageEntry->MessageId && CanScrollToLogDelegate.Execute(*PackageEntry->MessageId, [this](const FConcertLogEntry& Entry)
+		{
+			return SharedFilterLogEntry(Entry);
+		}, ErrorMessage);
+	}
+
+	void SPackageTransmissionTableRow::ScrollToLog() const
+	{
+		ScrollToLogDelegate.Execute(*PackageEntry->MessageId, [this](const FConcertLogEntry& Entry)
+		{
+			return SharedFilterLogEntry(Entry);
+		});
+	}
+	
+	bool SPackageTransmissionTableRow::SharedFilterLogEntry(const FConcertLogEntry& Entry) const
+	{
+		return Entry.Log.MessageAction == EConcertLogMessageAction::Process;
+	}
+}
+
+#undef LOCTEXT_NAMESPACE
\ No newline at end of file
diff --git a/Engine/Plugins/Developer/Concert/ConcertApp/MultiUserServer/Source/MultiUserServer/Private/Widgets/Clients/PackageTransmission/SPackageTransmissionTableRow.h b/Engine/Plugins/Developer/Concert/ConcertApp/MultiUserServer/Source/MultiUserServer/Private/Widgets/Clients/PackageTransmission/SPackageTransmissionTableRow.h
new file mode 100644
--- /dev/null
+++ b/Engine/Plugins/Developer/Concert/ConcertApp/MultiUserServer/Source/MultiUserServer/Private/Widgets/Clients/PackageTransmission/SPackageTransmissionTableRow.h
@@ -0,0 +1,70 @@
+// Copyright Epic Games, Inc. All Rights Reserved.
+
+#pragma once
+
+#include "CoreMinimal.h"
+#include "Model/PackageTransmissionEntry.h"
+#include "Widgets/DeclarativeSyntaxSupport.h"
+#include "Widgets/Clients/Logging/LogScrollingDelegates.h"
+#include "Widgets/Views/STableRow.h"
+#include "Widgets/Views/SListView.h"
+
+class FConcertLogTokenizer;
+
+namespace UE::MultiUserServer
+{
+	class FPackageTransmissionEntryTokenizer;
+	enum class EPackageTransmissionState : uint8;
+
+	class SPackageTransmissionTableRow : public SMultiColumnTableRow<TSharedPtr<FPackageTransmissionEntry>>
+	{
+	public:
+
+		static const FName TimeColumn;
+		static const FName OriginColumn;
+		static const FName DestinationColumn;
+		static const FName SizeColumn;
+		static const FName RevisionColumn;
+		static const FName PackagePathColumn;
+		static const FName PackageNameColumn;
+		static const FName TransmissionStateColumn;
+
+		static const TArray<FName> AllColumns;
+		static const TMap<FName, FText> ColumnsDisplayText;
+		
+		SLATE_BEGIN_ARGS(SPackageTransmissionTableRow)
+		{}
+			SLATE_ATTRIBUTE(FText, HighlightText)
+			SLATE_EVENT(FCanScrollToLog, CanScrollToLog)
+			SLATE_EVENT(FScrollToLog, ScrollToLog)
+		SLATE_END_ARGS()
+
+		void Construct(const FArguments& InArgs, TSharedPtr<FPackageTransmissionEntry> InPackageEntry, const TSharedRef<STableViewBase>& InOwnerTableView, TSharedRef<FPackageTransmissionEntryTokenizer> InTokenizer);
+
+		//~ Begin SMultiColumnTableRow Interface
+		virtual TSharedRef<SWidget> GenerateWidgetForColumn(const FName& ColumnName) override;
+		//~ End SMultiColumnTableRow Interface
+		
+	private:
+
+		TSharedPtr<FPackageTransmissionEntry> PackageEntry;
+		TSharedPtr<FPackageTransmissionEntryTokenizer> Tokenizer;
+
+		TAttribute<FText> HighlightText;
+		FCanScrollToLog CanScrollToLogDelegate;
+		FScrollToLog ScrollToLogDelegate;
+		
+		TSharedRef<SWidget> CreateTransmissionStateColumn() const;
+		
+		TSharedRef<SWidget> CreateTransmissionInProgressWidget() const;
+		TSharedRef<SWidget> CreateTransmissionSuccessWidget() const;
+		TSharedRef<SWidget> CreateTransmissionFailureWidget() const;
+		TSharedRef<SWidget> SharedCreateTransmissionCompletedWidget(EPackageTransmissionState AllowedTransmissionState, const FSlateBrush* ImageBrush) const;
+
+		bool CanScrollToLog() const { FText Dummy; return CanScrollToLog(Dummy); }
+		bool CanScrollToLog(FText& ErrorMessage) const;
+		void ScrollToLog() const;
+		bool SharedFilterLogEntry(const FConcertLogEntry& Entry) const;
+	};
+}
+
diff --git a/Engine/Plugins/Developer/Concert/ConcertApp/MultiUserServer/Source/MultiUserServer/Private/Widgets/Clients/PackageTransmission/SPackageTransmissionView.cpp b/Engine/Plugins/Developer/Concert/ConcertApp/MultiUserServer/Source/MultiUserServer/Private/Widgets/Clients/PackageTransmission/SPackageTransmissionView.cpp
new file mode 100644
--- /dev/null
+++ b/Engine/Plugins/Developer/Concert/ConcertApp/MultiUserServer/Source/MultiUserServer/Private/Widgets/Clients/PackageTransmission/SPackageTransmissionView.cpp
@@ -0,0 +1,47 @@
+// Copyright Epic Games, Inc. All Rights Reserved.
+
+#include "SPackageTransmissionView.h"
+
+#include "SPackageTransmissionTable.h"
+#include "Filter/FrontendPackageTransmissionFilter_TextSearch.h"
+#include "Filter/PackageTransmissionFilter_FrontendRoot.h"
+#include "Model/FilteredPackageTransmissionModel.h"
+
+#include "Widgets/SBoxPanel.h"
+
+namespace UE::MultiUserServer
+{
+	void SPackageTransmissionView::Construct(const FArguments& InArgs, TSharedRef<IPackageTransmissionEntrySource> InPackageEntrySource, TSharedRef<FPackageTransmissionEntryTokenizer> InTokenizer)
+	{
+		PackageEntrySource = MoveTemp(InPackageEntrySource);
+		RootFilter = MakeFilter(InTokenizer);
+		FilteredModel = MakeShared<FFilteredPackageTransmissionModel>(PackageEntrySource.ToSharedRef(), RootFilter.ToSharedRef());
+		
+		ChildSlot
+		[
+			SNew(SBorder)
+			.BorderImage(FAppStyle::Get().GetBrush("ToolPanel.GroupBorder"))
+			.BorderBackgroundColor(FSlateColor(FLinearColor(0.6, 0.6, 0.6)))
+			.Padding(2)
+			[
+				SNew(SVerticalBox)
+
+				+SVerticalBox::Slot()
+				.AutoHeight()
+				[
+					RootFilter->BuildFilterWidgets()
+				]
+				
+				+SVerticalBox::Slot()
+				[
+					SNew(SPackageTransmissionTable, FilteredModel.ToSharedRef(), InTokenizer)
+					.HighlightText_Lambda([this](){ return RootFilter->GetTextSearchFilter()->GetSearchText(); })
+					.CanScrollToLog(InArgs._CanScrollToLog)
+					.ScrollToLog(InArgs._ScrollToLog)
+					.TotalUnfilteredNum_Lambda([this](){ return PackageEntrySource->GetEntries().Num(); })
+				]
+			]
+		];
+	}
+}
+
diff --git a/Engine/Plugins/Developer/Concert/ConcertApp/MultiUserServer/Source/MultiUserServer/Private/Widgets/Clients/PackageTransmission/SPackageTransmissionView.h b/Engine/Plugins/Developer/Concert/ConcertApp/MultiUserServer/Source/MultiUserServer/Private/Widgets/Clients/PackageTransmission/SPackageTransmissionView.h
new file mode 100644
--- /dev/null
+++ b/Engine/Plugins/Developer/Concert/ConcertApp/MultiUserServer/Source/MultiUserServer/Private/Widgets/Clients/PackageTransmission/SPackageTransmissionView.h
@@ -0,0 +1,42 @@
+// Copyright Epic Games, Inc. All Rights Reserved.
+
+#pragma once
+
+#include "CoreMinimal.h"
+#include "Widgets/DeclarativeSyntaxSupport.h"
+#include "Widgets/SCompoundWidget.h"
+#include "Widgets/Clients/Logging/LogScrollingDelegates.h"
+
+class IConcertSyncServer;
+
+namespace UE::MultiUserServer
+{
+	class FFilteredPackageTransmissionModel;
+	class FPackageTransmissionEntryTokenizer;
+	class FPackageTransmissionFilter_FrontendRoot;
+	class IPackageTransmissionEntrySource;
+
+	/** Contains all UI for displaying package transmission, including filters, the table, and view options. */
+	class SPackageTransmissionView : public SCompoundWidget
+	{
+	public:
+
+		SLATE_BEGIN_ARGS(SPackageTransmissionView)
+		{}
+			SLATE_EVENT(FCanScrollToLog, CanScrollToLog)
+			SLATE_EVENT(FScrollToLog, ScrollToLog)
+		SLATE_END_ARGS()
+
+		void Construct(const FArguments& InArgs, TSharedRef<IPackageTransmissionEntrySource> InPackageEntrySource, TSharedRef<FPackageTransmissionEntryTokenizer> InTokenizer);
+
+	private:
+
+		TSharedPtr<IPackageTransmissionEntrySource> PackageEntrySource;
+		
+		TSharedPtr<FPackageTransmissionFilter_FrontendRoot> RootFilter;
+		TSharedPtr<FFilteredPackageTransmissionModel> FilteredModel;
+		
+		FCanScrollToLog CanScrollToLogDelegate;
+		FScrollToLog ScrollToLogDelegate;
+	};
+}
diff --git a/Engine/Plugins/Developer/Concert/ConcertApp/MultiUserServer/Source/MultiUserServer/Private/Widgets/Clients/PackageTransmission/Util/PackageTransmissionEntryTokenizer.cpp b/Engine/Plugins/Developer/Concert/ConcertApp/MultiUserServer/Source/MultiUserServer/Private/Widgets/Clients/PackageTransmission/Util/PackageTransmissionEntryTokenizer.cpp
new file mode 100644
--- /dev/null
+++ b/Engine/Plugins/Developer/Concert/ConcertApp/MultiUserServer/Source/MultiUserServer/Private/Widgets/Clients/PackageTransmission/Util/PackageTransmissionEntryTokenizer.cpp
@@ -0,0 +1,72 @@
+// Copyright Epic Games, Inc. All Rights Reserved.
+
+#include "PackageTransmissionEntryTokenizer.h"
+
+#include "ConcertFrontendUtils.h"
+#include "Math/UnitConversion.h"
+#include "Widgets/Clients/PackageTransmission/Model/PackageTransmissionEntry.h"
+#include "Settings/MultiUserServerPackageTransmissionSettings.h"
+#include "Widgets/Clients/Util/EndpointToUserNameCache.h"
+
+namespace UE::MultiUserServer
+{
+	FPackageTransmissionEntryTokenizer::FPackageTransmissionEntryTokenizer(TSharedRef<FEndpointToUserNameCache> EndpointInfoGetter)
+		: EndpointInfoGetter(MoveTemp(EndpointInfoGetter))
+	{}
+
+	FString FPackageTransmissionEntryTokenizer::TokenizeTime(const FPackageTransmissionEntry& Entry) const
+	{
+		return ConcertFrontendUtils::FormatTime(Entry.LocalStartTime, UMultiUserServerPackageTransmissionSettings::GetSettings()->TimestampTimeFormat).ToString();
+	}
+
+	FString FPackageTransmissionEntryTokenizer::TokenizeOrigin(const FPackageTransmissionEntry& Entry) const
+	{
+		switch (Entry.SendDirection)
+		{
+		case EPackageSendDirection::ClientToServer: return EndpointInfoGetter->GetEndpointDisplayString(Entry.ClientEndpointId);
+		case EPackageSendDirection::ServerToClient: return TEXT("Server");
+		default:
+			checkNoEntry();
+			return TEXT("n/a");
+		}
+	}
+
+	FString FPackageTransmissionEntryTokenizer::TokenizeDestination(const FPackageTransmissionEntry& Entry) const
+	{
+		switch (Entry.SendDirection)
+		{
+		case EPackageSendDirection::ServerToClient: return EndpointInfoGetter->GetEndpointDisplayString(Entry.ClientEndpointId);
+		case EPackageSendDirection::ClientToServer: return TEXT("Server");
+		default:
+			checkNoEntry();
+			return TEXT("n/a");
+		}
+	}
+
+	FString FPackageTransmissionEntryTokenizer::TokenizeSize(const FPackageTransmissionEntry& Entry) const
+	{
+		const FNumericUnit<uint64> Unit = FUnitConversion::QuantizeUnitsToBestFit(Entry.PackageNumBytes, EUnit::Bytes);
+		return FString::Printf(TEXT("%lld %s"), Unit.Value, FUnitConversion::GetUnitDisplayString(Unit.Units));
+	}
+
+	FString FPackageTransmissionEntryTokenizer::TokenizeRevision(const FPackageTransmissionEntry& Entry) const
+	{
+		return FString::FromInt(Entry.Revision);
+	}
+
+	FString FPackageTransmissionEntryTokenizer::TokenizePackagePath(const FPackageTransmissionEntry& Entry) const
+	{
+		return Entry.PackageInfo.PackageName.ToString();
+	}
+
+	FString FPackageTransmissionEntryTokenizer::TokenizePackageName(const FPackageTransmissionEntry& Entry) const
+	{
+		FString PackagePath = Entry.PackageInfo.PackageName.ToString();
+		int32 Index;
+		if (PackagePath.FindLastChar(TEXT('/'), Index) && ensure(Index != PackagePath.Len() - 1))
+		{
+			PackagePath.RightChopInline(Index + 1);
+		}
+		return PackagePath;
+	}
+}
diff --git a/Engine/Plugins/Developer/Concert/ConcertApp/MultiUserServer/Source/MultiUserServer/Private/Widgets/Clients/PackageTransmission/Util/PackageTransmissionEntryTokenizer.h b/Engine/Plugins/Developer/Concert/ConcertApp/MultiUserServer/Source/MultiUserServer/Private/Widgets/Clients/PackageTransmission/Util/PackageTransmissionEntryTokenizer.h
new file mode 100644
--- /dev/null
+++ b/Engine/Plugins/Developer/Concert/ConcertApp/MultiUserServer/Source/MultiUserServer/Private/Widgets/Clients/PackageTransmission/Util/PackageTransmissionEntryTokenizer.h
@@ -0,0 +1,35 @@
+// Copyright Epic Games, Inc. All Rights Reserved.
+
+#pragma once
+
+#include "CoreMinimal.h"
+
+class FEndpointToUserNameCache;
+
+namespace UE::MultiUserServer
+{
+	struct FPackageTransmissionEntry;
+	
+	/** Helps with printing FPackageTransmissionEntry to text so text search can search exactly what is displayed in the UI. */
+	class FPackageTransmissionEntryTokenizer
+	{
+	public:
+		
+		FPackageTransmissionEntryTokenizer(TSharedRef<FEndpointToUserNameCache> EndpointInfoGetter);
+
+		FString TokenizeTime(const FPackageTransmissionEntry& Entry) const;
+		FString TokenizeOrigin(const FPackageTransmissionEntry& Entry) const;
+		FString TokenizeDestination(const FPackageTransmissionEntry& Entry) const;
+		FString TokenizeSize(const FPackageTransmissionEntry& Entry) const;
+		FString TokenizeRevision(const FPackageTransmissionEntry& Entry) const;
+		FString TokenizePackagePath(const FPackageTransmissionEntry& Entry) const;
+		FString TokenizePackageName(const FPackageTransmissionEntry& Entry) const;
+
+	private:
+		
+		/** Used so we can look up client and server info (even after client has disconnected) */
+		TSharedRef<FEndpointToUserNameCache> EndpointInfoGetter;
+	};
+}
+
+
diff --git a/Engine/Plugins/Developer/Concert/ConcertApp/MultiUserServer/Source/MultiUserServer/Private/Widgets/Clients/SConcertClientsTabView.cpp b/Engine/Plugins/Developer/Concert/ConcertApp/MultiUserServer/Source/MultiUserServer/Private/Widgets/Clients/SConcertClientsTabView.cpp
--- a/Engine/Plugins/Developer/Concert/ConcertApp/MultiUserServer/Source/MultiUserServer/Private/Widgets/Clients/SConcertClientsTabView.cpp
+++ b/Engine/Plugins/Developer/Concert/ConcertApp/MultiUserServer/Source/MultiUserServer/Private/Widgets/Clients/SConcertClientsTabView.cpp
@@ -2,12 +2,12 @@
 
 #include "SConcertClientsTabView.h"
 
-#include "ConcertUtil.h"
-#include "IConcertServer.h"
 #include "IConcertSyncServer.h"
 #include "Logging/Filter/ConcertLogFilter_FrontendRoot.h"
 #include "Logging/Source/GlobalLogSource.h"
 #include "Logging/Util/ConcertLogTokenizer.h"
+#include "PackageTransmission/PackageTransmissionTabController.h"
+#include "PackageTransmission/Model/PackageTransmissionModel.h"
 #include "Util/EndpointToUserNameCache.h"
 #include "Widgets/Clients/Browser/SConcertNetworkBrowser.h"
 #include "Widgets/Clients/Browser/Models/ClientBrowserModel.h"
@@ -24,6 +24,7 @@
 
 const FName SConcertClientsTabView::ClientBrowserTabId("ClientBrowserTabId");
 const FName SConcertClientsTabView::GlobalLogTabId("GlobalLogTabId");
+const FName SConcertClientsTabView::PackageTransmissionTabId("PackageTransmissionTabId");
 
 void SConcertClientsTabView::Construct(const FArguments& InArgs, FName InStatusBarID, TSharedRef<IConcertSyncServer> InServer, TSharedRef<FGlobalLogSource> InLogBuffer)
 {
@@ -31,6 +32,7 @@ void SConcertClientsTabView::Construct(const FArguments& InArgs, FName InStatusB
 	LogBuffer = MoveTemp(InLogBuffer);
 	ClientInfoCache = MakeShared<FEndpointToUserNameCache>(Server->GetConcertServer());
 	LogTokenizer = MakeShared<FConcertLogTokenizer>(ClientInfoCache.ToSharedRef());
+	PackageTransmissionModel = MakeShared<UE::MultiUserServer::FPackageTransmissionModel>(Server.ToSharedRef());
 	
 	SConcertTabViewWithManagerBase::Construct(
 		SConcertTabViewWithManagerBase::FArguments()
@@ -40,7 +42,7 @@ void SConcertClientsTabView::Construct(const FArguments& InArgs, FName InStatusB
 		{
 			CreateTabs(InTabManager, InLayout, InArgs);
 		}))
-		.LayoutName("ConcertClientsTabView_v0.1"),
+		.LayoutName("ConcertClientsTabView_v0.2"),
 		InStatusBarID
 	);
 }
@@ -93,6 +95,24 @@ void SConcertClientsTabView::OpenClientLogTab(const FGuid& ClientMessageNodeId)
 	}
 }
 
+bool SConcertClientsTabView::CanScrollToLog(const FGuid& MessageId, FConcertLogEntryFilterFunc FilterFunc, FText& ErrorMessage) const
+{
+	const bool bCanScrollToLog = GlobalTransportLog->CanScrollToLog(MessageId, FilterFunc);
+	if (!bCanScrollToLog)
+	{
+		ErrorMessage = LOCTEXT("ScrollToLog.LogNotAvailable", "The log is filtered out.");
+		return false;
+	}
+
+	return true;
+}
+
+void SConcertClientsTabView::ScrollToLog(const FGuid& MessageId, FConcertLogEntryFilterFunc FilterFunc) const
+{
+	OpenGlobalLogTab();
+	GlobalTransportLog->ScrollToLog(MessageId, FilterFunc);
+}
+
 void SConcertClientsTabView::OpenGlobalLogTab() const
 {
 	GetTabManager()->TryInvokeTab(GlobalLogTabId);
@@ -118,6 +138,15 @@ TSharedPtr<SDockTab> SConcertClientsTabView::GetGlobalLogTab() const
 
 void SConcertClientsTabView::CreateTabs(const TSharedRef<FTabManager>& InTabManager, const TSharedRef<FTabManager::FLayout>& InLayout, const FArguments& InArgs)
 {
+	MainPackageTransmissionTab = MakeShared<UE::MultiUserServer::FPackageTransmissionTabController>(
+		PackageTransmissionTabId,
+		InTabManager,
+		PackageTransmissionModel.ToSharedRef(),
+		ClientInfoCache.ToSharedRef(),
+		UE::MultiUserServer::FCanScrollToLog::CreateSP(this, &SConcertClientsTabView::CanScrollToLog),
+		UE::MultiUserServer::FScrollToLog::CreateSP(this, &SConcertClientsTabView::ScrollToLog)
+		);
+	
 	InTabManager->RegisterTabSpawner(ClientBrowserTabId, FOnSpawnTab::CreateSP(this, &SConcertClientsTabView::SpawnClientBrowserTab))
 		.SetDisplayName(LOCTEXT("ClientBrowserTabLabel", "Clients"));
 	InTabManager->RegisterTabSpawner(GlobalLogTabId, FOnSpawnTab::CreateSP(this, &SConcertClientsTabView::SpawnGlobalLogTab))
@@ -128,14 +157,27 @@ void SConcertClientsTabView::CreateTabs(const TSharedRef<FTabManager>& InTabMana
 				->SetOrientation(Orient_Vertical)
 				->Split
 				(
-					FTabManager::NewStack()
+					FTabManager::NewSplitter()
 					->SetSizeCoefficient(0.5f)
-					->AddTab(ClientBrowserTabId, ETabState::OpenedTab)
+					->SetOrientation(Orient_Horizontal)
+					->Split
+					(
+						FTabManager::NewStack()
+						->SetSizeCoefficient(0.65f)
+						->AddTab(ClientBrowserTabId, ETabState::OpenedTab)
+					)
+					->Split
+					(
+						FTabManager::NewStack()
+						->SetSizeCoefficient(0.35f)
+						->AddTab(PackageTransmissionTabId, ETabState::OpenedTab)
+							
+					)
 				)
 				->Split
 				(
-					FTabManager::NewStack()
-					->SetSizeCoefficient(0.5f)
+				FTabManager::NewStack()
+					->SetSizeCoefficient(0.6f)
 					->AddTab(GlobalLogTabId, ETabState::OpenedTab)
 				)
 		);
@@ -167,7 +209,7 @@ TSharedRef<SDockTab> SConcertClientsTabView::SpawnGlobalLogTab(const FSpawnTabAr
 		.Label(LOCTEXT("GlobalLogTabLabel", "Global Log"))
 		.TabRole(PanelTab)
 		[
-			SNew(SConcertTransportLog, LogBuffer.ToSharedRef(), ClientInfoCache.ToSharedRef(), LogTokenizer.ToSharedRef())
+			SAssignNew(GlobalTransportLog, SConcertTransportLog, LogBuffer.ToSharedRef(), ClientInfoCache.ToSharedRef(), LogTokenizer.ToSharedRef())
 			.Filter(UE::MultiUserServer::MakeGlobalLogFilter(LogTokenizer.ToSharedRef()))
 		]; 
 }
diff --git a/Engine/Plugins/Developer/Concert/ConcertApp/MultiUserServer/Source/MultiUserServer/Private/Widgets/Clients/SConcertClientsTabView.h b/Engine/Plugins/Developer/Concert/ConcertApp/MultiUserServer/Source/MultiUserServer/Private/Widgets/Clients/SConcertClientsTabView.h
--- a/Engine/Plugins/Developer/Concert/ConcertApp/MultiUserServer/Source/MultiUserServer/Private/Widgets/Clients/SConcertClientsTabView.h
+++ b/Engine/Plugins/Developer/Concert/ConcertApp/MultiUserServer/Source/MultiUserServer/Private/Widgets/Clients/SConcertClientsTabView.h
@@ -3,13 +3,18 @@
 #pragma once
 
 #include "CoreMinimal.h"
+#include "Logging/ConcertLogEntry.h"
 #include "Widgets/DeclarativeSyntaxSupport.h"
 #include "Widgets/SConcertTabViewWithManagerBase.h"
 
+class SConcertTransportLog;
+
 namespace UE
 {
 	namespace MultiUserServer
 	{
+		class FPackageTransmissionTabController;
+		class FPackageTransmissionModel;
 		class SConcertNetworkBrowser;
 	}
 }
@@ -33,6 +38,7 @@ public:
 	
 	static const FName ClientBrowserTabId;
 	static const FName GlobalLogTabId;
+	static const FName PackageTransmissionTabId;
 
 	SLATE_BEGIN_ARGS(SConcertClientsTabView)
 	{}
@@ -46,6 +52,11 @@ public:
 	void OpenGlobalLogTab() const;
 	void CloseGlobalLogTab() const;
 	void OpenClientLogTab(const FGuid& ClientMessageNodeId) const;
+
+	/** Whether it is possible to call ScrollToLog. */
+	bool CanScrollToLog(const FGuid& MessageId, FConcertLogEntryFilterFunc FilterFunc, FText& ErrorMessage) const;
+	/** Scrolls to the log in the global log tab. */
+	void ScrollToLog(const FGuid& MessageId, FConcertLogEntryFilterFunc FilterFunc) const;
 	
 	bool IsGlobalLogOpen() const;
 	TSharedPtr<SDockTab> GetGlobalLogTab() const;
@@ -61,13 +72,15 @@ private:
 	TSharedPtr<FEndpointToUserNameCache> ClientInfoCache;
 	/** Used by various systems to convert logs to text */
 	TSharedPtr<FConcertLogTokenizer> LogTokenizer;
+	
+	/** Knows about when packages are sent and received by the server */
+	TSharedPtr<UE::MultiUserServer::FPackageTransmissionModel> PackageTransmissionModel;
+	TSharedPtr<UE::MultiUserServer::FPackageTransmissionTabController> MainPackageTransmissionTab;
 
 	TSharedPtr<UE::MultiUserServer::SConcertNetworkBrowser> ClientBrowser;
-	
-	/** Used to overlay EnableLoggingPrompt over the tabs */
-	TSharedPtr<SOverlay> EnableLoggingPromptOverlay;
-	/** Reminds the user to enable logging */
-	TSharedPtr<SPromptConcertLoggingEnabled> EnableLoggingPrompt;
+
+	/** The widget inside of the global log tab */
+	TSharedPtr<SConcertTransportLog> GlobalTransportLog;
 	
 	void CreateTabs(const TSharedRef<FTabManager>& InTabManager, const TSharedRef<FTabManager::FLayout>& InLayout, const FArguments& InArgs);
 	TSharedRef<SDockTab> SpawnClientBrowserTab(const FSpawnTabArgs& InTabArgs);
diff --git a/Engine/Plugins/Developer/Concert/ConcertApp/MultiUserServer/Source/MultiUserServer/Private/Widgets/Clients/Util/EndpointToUserNameCache.cpp b/Engine/Plugins/Developer/Concert/ConcertApp/MultiUserServer/Source/MultiUserServer/Private/Widgets/Clients/Util/EndpointToUserNameCache.cpp
--- a/Engine/Plugins/Developer/Concert/ConcertApp/MultiUserServer/Source/MultiUserServer/Private/Widgets/Clients/Util/EndpointToUserNameCache.cpp
+++ b/Engine/Plugins/Developer/Concert/ConcertApp/MultiUserServer/Source/MultiUserServer/Private/Widgets/Clients/Util/EndpointToUserNameCache.cpp
@@ -106,6 +106,21 @@ TOptional<FEndpointToUserNameCache::FNodeEndpointId> FEndpointToUserNameCache::T
 	return {};
 }
 
+FString FEndpointToUserNameCache::GetEndpointDisplayString(const FGuid& EndpointId) const
+{
+	if (IsServerEndpoint(EndpointId))
+	{
+		return FString("Server");		
+	}
+
+	if (const TOptional<FConcertClientInfo> ClientInfo = GetClientInfo(EndpointId))
+	{
+		return ClientInfo->DisplayName;
+	}
+	
+	return EndpointId.ToString(EGuidFormats::DigitsWithHyphens);
+}
+
 void FEndpointToUserNameCache::OnLiveSessionCreated(bool bSuccess, const IConcertServer& InServer, TSharedRef<IConcertServerSession> InLiveSession)
 {
 	RegisterLiveSession(InLiveSession);
diff --git a/Engine/Plugins/Developer/Concert/ConcertApp/MultiUserServer/Source/MultiUserServer/Private/Widgets/Clients/Util/EndpointToUserNameCache.h b/Engine/Plugins/Developer/Concert/ConcertApp/MultiUserServer/Source/MultiUserServer/Private/Widgets/Clients/Util/EndpointToUserNameCache.h
--- a/Engine/Plugins/Developer/Concert/ConcertApp/MultiUserServer/Source/MultiUserServer/Private/Widgets/Clients/Util/EndpointToUserNameCache.h
+++ b/Engine/Plugins/Developer/Concert/ConcertApp/MultiUserServer/Source/MultiUserServer/Private/Widgets/Clients/Util/EndpointToUserNameCache.h
@@ -37,6 +37,9 @@ public:
 	/** Translates a Concert endpoint Id to the corresponding messaging node Id */
 	TOptional<FNodeEndpointId> TranslateEndpointIdToNodeId(const FGuid& EndpointId) const;
 
+	/** Gets either the client's name or "Server" depending on type of endpoint this is */
+	FString GetEndpointDisplayString(const FGuid& EndpointId) const;
+
 private:
 
 	/** Used to unsubscribe when we're destroyed */
diff --git a/Engine/Plugins/Developer/Concert/ConcertApp/MultiUserServer/Source/MultiUserServer/Private/Window/ConcertServerWindowController.cpp b/Engine/Plugins/Developer/Concert/ConcertApp/MultiUserServer/Source/MultiUserServer/Private/Window/ConcertServerWindowController.cpp
--- a/Engine/Plugins/Developer/Concert/ConcertApp/MultiUserServer/Source/MultiUserServer/Private/Window/ConcertServerWindowController.cpp
+++ b/Engine/Plugins/Developer/Concert/ConcertApp/MultiUserServer/Source/MultiUserServer/Private/Window/ConcertServerWindowController.cpp
@@ -46,7 +46,7 @@ TSharedRef<SWindow> FConcertServerWindowController::CreateWindow()
 	const float DPIScaleFactor = FPlatformApplicationMisc::GetDPIScaleFactorAtPoint(DisplayMetrics.PrimaryDisplayWorkAreaRect.Left, DisplayMetrics.PrimaryDisplayWorkAreaRect.Top);
 
 	const bool bEmbedTitleAreaContent = false;
-	const FVector2D ClientSize(960.0f * DPIScaleFactor, 750.0f * DPIScaleFactor);
+	const FVector2D ClientSize(1000.0f * DPIScaleFactor, 800.0f * DPIScaleFactor);
 	TSharedRef<SWindow> RootWindowRef = SNew(SWindow)
 		.Title(LOCTEXT("WindowTitle", "Unreal Multi User Server"))
 		.CreateTitleBar(!bEmbedTitleAreaContent)
diff --git a/Engine/Plugins/Developer/Concert/ConcertSync/ConcertSyncCore/Source/ConcertSyncCore/Public/ConcertPackageEvents.h b/Engine/Plugins/Developer/Concert/ConcertSync/ConcertSyncCore/Source/ConcertSyncCore/Public/ConcertPackageEvents.h
--- a/Engine/Plugins/Developer/Concert/ConcertSync/ConcertSyncCore/Source/ConcertSyncCore/Public/ConcertPackageEvents.h
+++ b/Engine/Plugins/Developer/Concert/ConcertSync/ConcertSyncCore/Source/ConcertSyncCore/Public/ConcertPackageEvents.h
@@ -43,7 +43,7 @@ namespace UE::ConcertSyncCore::ConcertPackageEvents
 		FGuid MessageId;
 	};
 
-	struct FConcertRejectSendPackageParams : FConcertBaseSendPackageParams
+	struct FConcertRejectSendPackageParams : FConcertFinishSendPackageParams
 	{};
 	
 	DECLARE_MULTICAST_DELEGATE_OneParam(FConcertBeginSendPackageDelegate, const FConcertBeginSendPackageParams&);
diff --git a/Engine/Plugins/Developer/Concert/ConcertSync/ConcertSyncServer/Source/ConcertSyncServer/Private/ConcertServerWorkspace.cpp b/Engine/Plugins/Developer/Concert/ConcertSync/ConcertSyncServer/Source/ConcertSyncServer/Private/ConcertServerWorkspace.cpp
--- a/Engine/Plugins/Developer/Concert/ConcertSync/ConcertSyncServer/Source/ConcertSyncServer/Private/ConcertServerWorkspace.cpp
+++ b/Engine/Plugins/Developer/Concert/ConcertSync/ConcertSyncServer/Source/ConcertSyncServer/Private/ConcertServerWorkspace.cpp
@@ -422,7 +422,8 @@ void FConcertServerWorkspace::HandlePackageUpdateEvent(const FConcertSessionCont
 	{
 		if (Event.Package.HasPackageData())
 		{
-			UE::ConcertSyncCore::ConcertPackageEvents::OnRejectRemoteSendPackage().Broadcast({{ Event.TransmissionId, Event.Package.Info, Context.SourceEndpointId }});
+			const UE::ConcertSyncCore::ConcertPackageEvents::FConcertRejectSendPackageParams Params {{{ Event.TransmissionId, Event.Package.Info, Context.SourceEndpointId }, Context.MessageId}};
+			UE::ConcertSyncCore::ConcertPackageEvents::OnRejectRemoteSendPackage().Broadcast(Params);
 		}
 		
 		// If the client didn't have the lock, then queue a rejection event so that the client will re-load the head-revision of the package
