Refactor highlight text to be polling instead of event based

diff --git a/Engine/Plugins/Developer/Concert/ConcertApp/MultiUserServer/Source/MultiUserServer/Private/Widgets/Clients/Logging/Filter/ConcertFrontendLogFilter_TextSearch.cpp b/Engine/Plugins/Developer/Concert/ConcertApp/MultiUserServer/Source/MultiUserServer/Private/Widgets/Clients/Logging/Filter/ConcertFrontendLogFilter_TextSearch.cpp
--- a/Engine/Plugins/Developer/Concert/ConcertApp/MultiUserServer/Source/MultiUserServer/Private/Widgets/Clients/Logging/Filter/ConcertFrontendLogFilter_TextSearch.cpp
+++ b/Engine/Plugins/Developer/Concert/ConcertApp/MultiUserServer/Source/MultiUserServer/Private/Widgets/Clients/Logging/Filter/ConcertFrontendLogFilter_TextSearch.cpp
@@ -30,7 +30,7 @@ FConcertFrontendLogFilter_TextSearch::FConcertFrontendLogFilter_TextSearch(TShar
 	ChildSlot = SNew(SSearchBox)
 		.OnTextChanged_Lambda([this](const FText& NewSearchText)
 		{
-			OnSearchTextChanged().Broadcast(NewSearchText);
+			SearchText = NewSearchText;
 			Implementation.SetRawFilterText(NewSearchText);
 		})
 		.DelayChangeNotificationsWhileTyping(true);
diff --git a/Engine/Plugins/Developer/Concert/ConcertApp/MultiUserServer/Source/MultiUserServer/Private/Widgets/Clients/Logging/Filter/ConcertFrontendLogFilter_TextSearch.h b/Engine/Plugins/Developer/Concert/ConcertApp/MultiUserServer/Source/MultiUserServer/Private/Widgets/Clients/Logging/Filter/ConcertFrontendLogFilter_TextSearch.h
--- a/Engine/Plugins/Developer/Concert/ConcertApp/MultiUserServer/Source/MultiUserServer/Private/Widgets/Clients/Logging/Filter/ConcertFrontendLogFilter_TextSearch.h
+++ b/Engine/Plugins/Developer/Concert/ConcertApp/MultiUserServer/Source/MultiUserServer/Private/Widgets/Clients/Logging/Filter/ConcertFrontendLogFilter_TextSearch.h
@@ -41,11 +41,9 @@ public:
 
 	FConcertFrontendLogFilter_TextSearch(TSharedRef<FConcertLogTokenizer> Tokenizer);
 
-	DECLARE_MULTICAST_DELEGATE_OneParam(FOnSearchTextChanged, const FText& /*NewSearchText*/);
-	FOnSearchTextChanged& OnSearchTextChanged() { return OnSearchTextChangedEvent; }
+	FText GetSearchText() const { return SearchText; }
 
 private:
 
-	/** Useful to let external subscriber handle text highlighting */
-	FOnSearchTextChanged OnSearchTextChangedEvent;
+	FText SearchText;
 };
diff --git a/Engine/Plugins/Developer/Concert/ConcertApp/MultiUserServer/Source/MultiUserServer/Private/Widgets/Clients/Logging/Filter/ConcertLogFilter.cpp b/Engine/Plugins/Developer/Concert/ConcertApp/MultiUserServer/Source/MultiUserServer/Private/Widgets/Clients/Logging/Filter/ConcertLogFilter.cpp
deleted file mode 100644
diff --git a/Engine/Plugins/Developer/Concert/ConcertApp/MultiUserServer/Source/MultiUserServer/Private/Widgets/Clients/Logging/SConcertTransportLog.cpp b/Engine/Plugins/Developer/Concert/ConcertApp/MultiUserServer/Source/MultiUserServer/Private/Widgets/Clients/Logging/SConcertTransportLog.cpp
--- a/Engine/Plugins/Developer/Concert/ConcertApp/MultiUserServer/Source/MultiUserServer/Private/Widgets/Clients/Logging/SConcertTransportLog.cpp
+++ b/Engine/Plugins/Developer/Concert/ConcertApp/MultiUserServer/Source/MultiUserServer/Private/Widgets/Clients/Logging/SConcertTransportLog.cpp
@@ -36,7 +36,7 @@ void SConcertTransportLog::Construct(const FArguments& InArgs, TSharedRef<IConce
 	PagedLogList = MakeShared<FPagedFilteredConcertLogList>(MoveTemp(LogSource), InArgs._Filter);
 	EndpointCache = MoveTemp(InEndpointCache);
 	LogTokenizer = MoveTemp(InLogTokenizer);
-	HighlightText = MakeShared<FText>();
+	Filter = InArgs._Filter;
 	
 	ChildSlot
 	[
@@ -77,10 +77,6 @@ void SConcertTransportLog::Construct(const FArguments& InArgs, TSharedRef<IConce
 
 	// Subscribe to events
 	PagedLogList->OnPageViewChanged().AddSP(this, &SConcertTransportLog::OnPageViewChanged);
-	if (InArgs._Filter)
-	{
-		InArgs._Filter->GetTextSearchFilter()->OnSearchTextChanged().AddSP(this, &SConcertTransportLog::OnSearchTextChanged);
-	}
 	
 	UMultiUserServerColumnVisibilitySettings::GetSettings()->OnTransportLogColumnVisibility().AddSP(this, &SConcertTransportLog::OnColumnVisibilitySettingsChanged);
 	UE::ConcertSharedSlate::RestoreColumnVisibilityState(HeaderRow.ToSharedRef(), UMultiUserServerColumnVisibilitySettings::GetSettings()->GetTransportLogColumnVisibility());
@@ -89,20 +85,20 @@ void SConcertTransportLog::Construct(const FArguments& InArgs, TSharedRef<IConce
 	OnConcertLoggingEnabledChanged(ConcertTransportEvents::IsLoggingEnabled());
 }
 
-bool SConcertTransportLog::CanScrollToLog(const FGuid& MessageId, FConcertLogEntryFilterFunc Filter) const
+bool SConcertTransportLog::CanScrollToLog(const FGuid& MessageId, FConcertLogEntryFilterFunc FilterFunc) const
 {
-	return PagedLogList->GetFilteredLogsWithId(MessageId).ContainsByPredicate([Filter](const TSharedPtr<FConcertLogEntry>& LogEntry)
+	return PagedLogList->GetFilteredLogsWithId(MessageId).ContainsByPredicate([FilterFunc](const TSharedPtr<FConcertLogEntry>& LogEntry)
 	{
-		return Filter(*LogEntry);
+		return FilterFunc(*LogEntry);
 	});
 }
 
-void SConcertTransportLog::ScrollToLog(const FGuid& MessageId, FConcertLogEntryFilterFunc Filter) const
+void SConcertTransportLog::ScrollToLog(const FGuid& MessageId, FConcertLogEntryFilterFunc FilterFunc) const
 {
 	// Could be optimised further by having FPagedFilteredConcertLogList cache the index but it's fine - ScrollToLog is called very infrequently (when the user presses a button)
-	const int32 Index = PagedLogList->GetFilteredLogs().IndexOfByPredicate([MessageId, Filter](const TSharedPtr<FConcertLogEntry>& Entry)
+	const int32 Index = PagedLogList->GetFilteredLogs().IndexOfByPredicate([MessageId, FilterFunc](const TSharedPtr<FConcertLogEntry>& Entry)
 	{
-		return Entry->Log.MessageId == MessageId && Filter(*Entry);
+		return Entry->Log.MessageId == MessageId && FilterFunc(*Entry);
 	});
 	ScrollToLog(Index);
 }
@@ -196,7 +192,8 @@ TSharedRef<ITableRow> SConcertTransportLog::OnGenerateActivityRowWidget(TSharedP
 		? OriginInfo->AvatarColor
 		: DestinationInfo.IsSet() ? DestinationInfo->AvatarColor : FLinearColor::Black;
 	
-	return SNew(SConcertTransportLogRow, Item, OwnerTable, LogTokenizer.ToSharedRef(), HighlightText.ToSharedRef())
+	return SNew(SConcertTransportLogRow, Item, OwnerTable, LogTokenizer.ToSharedRef())
+		.HighlightText_Lambda([this](){ return Filter.IsValid() ? Filter->GetTextSearchFilter()->GetSearchText() : FText(); })
 		.AvatarColor(AvatarColor)
 		.CanScrollToAckLog(this, &SConcertTransportLog::CanScrollToAckLog)
 		.CanScrollToAckedLog(this, &SConcertTransportLog::CanScrollToAckedLog)
@@ -343,11 +340,6 @@ void SConcertTransportLog::OnPageViewChanged(const TArray<TSharedPtr<FConcertLog
 	}
 }
 
-void SConcertTransportLog::OnSearchTextChanged(const FText& NewSearchText)
-{
-	*HighlightText = NewSearchText;
-}
-
 void SConcertTransportLog::OnColumnVisibilitySettingsChanged(const FColumnVisibilitySnapshot& ColumnSnapshot)
 {
 	TGuardValue<bool> GuardValue(bIsUpdatingColumnVisibility, true);
diff --git a/Engine/Plugins/Developer/Concert/ConcertApp/MultiUserServer/Source/MultiUserServer/Private/Widgets/Clients/Logging/SConcertTransportLog.h b/Engine/Plugins/Developer/Concert/ConcertApp/MultiUserServer/Source/MultiUserServer/Private/Widgets/Clients/Logging/SConcertTransportLog.h
--- a/Engine/Plugins/Developer/Concert/ConcertApp/MultiUserServer/Source/MultiUserServer/Private/Widgets/Clients/Logging/SConcertTransportLog.h
+++ b/Engine/Plugins/Developer/Concert/ConcertApp/MultiUserServer/Source/MultiUserServer/Private/Widgets/Clients/Logging/SConcertTransportLog.h
@@ -41,8 +41,8 @@ public:
 
 	void Construct(const FArguments& InArgs, TSharedRef<IConcertLogSource> LogSource, TSharedRef<FEndpointToUserNameCache> InEndpointCache, TSharedRef<FConcertLogTokenizer> InLogTokenizer);
 
-	bool CanScrollToLog(const FGuid& MessageId, FConcertLogEntryFilterFunc Filter) const;
-	void ScrollToLog(const FGuid& MessageId, FConcertLogEntryFilterFunc Filter) const;
+	bool CanScrollToLog(const FGuid& MessageId, FConcertLogEntryFilterFunc FilterFunc) const;
+	void ScrollToLog(const FGuid& MessageId, FConcertLogEntryFilterFunc FilterFunc) const;
 	
 private:
 	
@@ -58,8 +58,7 @@ private:
 	/** Used by various systems to convert logs to text */
 	TSharedPtr<FConcertLogTokenizer> LogTokenizer;
 
-	/** Updates to be the content of the search text. Shared with all rows. */
-	TSharedPtr<FText> HighlightText;
+	TSharedPtr<FConcertLogFilter_FrontendRoot> Filter;
 	
 	/** Lists the logs */
 	TSharedPtr<SListView<TSharedPtr<FConcertLogEntry>>> LogView;
@@ -82,7 +81,6 @@ private:
 	void ExtendViewOptions(FMenuBuilder& MenuBuilder);
 	void OnFilterMenuChecked();
 	void OnPageViewChanged(const TArray<TSharedPtr<FConcertLogEntry>>&);
-	void OnSearchTextChanged(const FText& NewSearchText);
 	void OnColumnVisibilitySettingsChanged(const FColumnVisibilitySnapshot& ColumnSnapshot);
 	
 	void OnConcertLoggingEnabledChanged(bool bNewEnabled);
diff --git a/Engine/Plugins/Developer/Concert/ConcertApp/MultiUserServer/Source/MultiUserServer/Private/Widgets/Clients/Logging/SConcertTransportLogRow.cpp b/Engine/Plugins/Developer/Concert/ConcertApp/MultiUserServer/Source/MultiUserServer/Private/Widgets/Clients/Logging/SConcertTransportLogRow.cpp
--- a/Engine/Plugins/Developer/Concert/ConcertApp/MultiUserServer/Source/MultiUserServer/Private/Widgets/Clients/Logging/SConcertTransportLogRow.cpp
+++ b/Engine/Plugins/Developer/Concert/ConcertApp/MultiUserServer/Source/MultiUserServer/Private/Widgets/Clients/Logging/SConcertTransportLogRow.cpp
@@ -20,11 +20,11 @@
 
 #define LOCTEXT_NAMESPACE "UnrealMultiUserUI.SConcertTransportLogRow"
 
-void SConcertTransportLogRow::Construct(const FArguments& InArgs, TSharedPtr<FConcertLogEntry> InLogEntry, const TSharedRef<STableViewBase>& InOwnerTableView, TSharedRef<FConcertLogTokenizer> InTokenizer, TSharedRef<FText> InHighlightText)
+void SConcertTransportLogRow::Construct(const FArguments& InArgs, TSharedPtr<FConcertLogEntry> InLogEntry, const TSharedRef<STableViewBase>& InOwnerTableView, TSharedRef<FConcertLogTokenizer> InTokenizer)
 {
 	LogEntry = MoveTemp(InLogEntry);
 	Tokenizer = MoveTemp(InTokenizer);
-	HighlightText = MoveTemp(InHighlightText);
+	HighlightText = InArgs._HighlightText;
 	
 	AvatarColor = InArgs._AvatarColor;
 	CanScrollToAckLogFunc = InArgs._CanScrollToAckLog;
@@ -88,7 +88,7 @@ TSharedRef<SWidget> SConcertTransportLogRow::CreateDefaultColumn(const FName& Pr
 				
 			return FText::GetEmpty();
 		})
-		.HighlightText_Lambda([this](){ return *HighlightText.Get(); });
+		.HighlightText(HighlightText);
 }
 
 TSharedRef<SWidget> SConcertTransportLogRow::CreateAckColumn(const FName& PropertyName)
diff --git a/Engine/Plugins/Developer/Concert/ConcertApp/MultiUserServer/Source/MultiUserServer/Private/Widgets/Clients/Logging/SConcertTransportLogRow.h b/Engine/Plugins/Developer/Concert/ConcertApp/MultiUserServer/Source/MultiUserServer/Private/Widgets/Clients/Logging/SConcertTransportLogRow.h
--- a/Engine/Plugins/Developer/Concert/ConcertApp/MultiUserServer/Source/MultiUserServer/Private/Widgets/Clients/Logging/SConcertTransportLogRow.h
+++ b/Engine/Plugins/Developer/Concert/ConcertApp/MultiUserServer/Source/MultiUserServer/Private/Widgets/Clients/Logging/SConcertTransportLogRow.h
@@ -18,6 +18,7 @@ public:
 
 	SLATE_BEGIN_ARGS(SConcertTransportLogRow)
 	{}
+		SLATE_ATTRIBUTE(FText, HighlightText)
 		SLATE_ARGUMENT(FLinearColor, AvatarColor)
 		SLATE_EVENT(FCanScrollToLog, CanScrollToAckLog)
 		SLATE_EVENT(FCanScrollToLog, CanScrollToAckedLog)
@@ -25,7 +26,7 @@ public:
 		SLATE_EVENT(FScrollToLog, ScrollToAckedLog)
 	SLATE_END_ARGS()
 
-	void Construct(const FArguments& InArgs, TSharedPtr<FConcertLogEntry> InLogEntry, const TSharedRef<STableViewBase>& InOwnerTableView, TSharedRef<FConcertLogTokenizer> InTokenizer, TSharedRef<FText> InHighlightText);
+	void Construct(const FArguments& InArgs, TSharedPtr<FConcertLogEntry> InLogEntry, const TSharedRef<STableViewBase>& InOwnerTableView, TSharedRef<FConcertLogTokenizer> InTokenizer);
 
 	//~ Begin SMultiColumnTableRow Interface
 	virtual TSharedRef<SWidget> GenerateWidgetForColumn(const FName& ColumnName) override;
@@ -37,8 +38,7 @@ private:
 	/** Used to convert some members into display strings */
 	TSharedPtr<FConcertLogTokenizer> Tokenizer;
 
-	/** Owned by SConcertTransportLog. Updated with the search text. */
-	TSharedPtr<FText> HighlightText;
+	TAttribute<FText> HighlightText;
 
 	FLinearColor AvatarColor;
 	FCanScrollToLog CanScrollToAckLogFunc;
