Make package and log filters use templated base class. Also add several classes to UE::MultiUserServer namespace.

diff --git a/Engine/Plugins/Developer/Concert/ConcertApp/MultiUserServer/Source/MultiUserServer/Private/Widgets/Browser/ConcertServerSessionBrowserController.cpp b/Engine/Plugins/Developer/Concert/ConcertApp/MultiUserServer/Source/MultiUserServer/Private/Widgets/Browser/ConcertServerSessionBrowserController.cpp
--- a/Engine/Plugins/Developer/Concert/ConcertApp/MultiUserServer/Source/MultiUserServer/Private/Widgets/Browser/ConcertServerSessionBrowserController.cpp
+++ b/Engine/Plugins/Developer/Concert/ConcertApp/MultiUserServer/Source/MultiUserServer/Private/Widgets/Browser/ConcertServerSessionBrowserController.cpp
@@ -16,334 +16,337 @@
 #include "Window/ConcertServerTabs.h"
 #include "Window/ConcertServerWindowController.h"
 
-#define LOCTEXT_NAMESPACE "UnrealMultiUserUI"
+#define LOCTEXT_NAMESPACE "UnrealMultiUserUI.FConcertServerSessionBrowserController"
 
-int32 FConcertServerSessionBrowserController::GetNumConnectedClients(const FGuid& SessionId) const
+namespace UE::MultiUserServer
 {
-	TSharedPtr<IConcertServerSession> Session = ServerInstance->GetConcertServer()->GetLiveSession(SessionId);
-	return ensure(Session) ? Session->GetSessionClients().Num() : 0;
-}
+	int32 FConcertServerSessionBrowserController::GetNumConnectedClients(const FGuid& SessionId) const
+	{
+		TSharedPtr<IConcertServerSession> Session = ServerInstance->GetConcertServer()->GetLiveSession(SessionId);
+		return ensure(Session) ? Session->GetSessionClients().Num() : 0;
+	}
 
-FConcertServerSessionBrowserController::~FConcertServerSessionBrowserController()
-{
-	ConcertServerEvents::OnLiveSessionCreated().RemoveAll(this);
-	ConcertServerEvents::OnLiveSessionDestroyed().RemoveAll(this);
-	ConcertServerEvents::OnArchivedSessionCreated().RemoveAll(this);
-	ConcertServerEvents::OnArchivedSessionDestroyed().RemoveAll(this);
-}
+	FConcertServerSessionBrowserController::~FConcertServerSessionBrowserController()
+	{
+		ConcertServerEvents::OnLiveSessionCreated().RemoveAll(this);
+		ConcertServerEvents::OnLiveSessionDestroyed().RemoveAll(this);
+		ConcertServerEvents::OnArchivedSessionCreated().RemoveAll(this);
+		ConcertServerEvents::OnArchivedSessionDestroyed().RemoveAll(this);
+	}
 
-void FConcertServerSessionBrowserController::Init(const FConcertComponentInitParams& Params)
-{
-	ServerInstance = Params.Server;
-	Owner = Params.WindowController;
-	
-	FGlobalTabmanager::Get()->RegisterTabSpawner(
-			ConcertServerTabs::GetSessionBrowserTabId(),
-			FOnSpawnTab::CreateRaw(this, &FConcertServerSessionBrowserController::SpawnSessionBrowserTab)
-		)
-		.SetDisplayName(LOCTEXT("SessionBrowserTabTitleLong", "Session Browser"))
-		.SetTooltipText(LOCTEXT("SessionBrowserTooltipText", "A section to browse, start, archive, and restore server sessions."))
-		.SetIcon(FSlateIcon(FConcertServerStyle::GetStyleSetName(), TEXT("Concert.MultiUser"))
-	);
-	Params.MainStack->AddTab(ConcertServerTabs::GetSessionBrowserTabId(), ETabState::OpenedTab)
-		->SetForegroundTab(ConcertServerTabs::GetSessionBrowserTabId());
-
-	ConcertServerEvents::OnLiveSessionCreated().AddSP(this, &FConcertServerSessionBrowserController::OnLiveSessionCreated);
-	ConcertServerEvents::OnLiveSessionDestroyed().AddSP(this, &FConcertServerSessionBrowserController::OnLiveSessionDestroyed);
-	ConcertServerEvents::OnArchivedSessionCreated().AddSP(this, &FConcertServerSessionBrowserController::OnArchivedSessionCreated);
-	ConcertServerEvents::OnArchivedSessionDestroyed().AddSP(this, &FConcertServerSessionBrowserController::OnArchivedSessionDestroyed);
-}
+	void FConcertServerSessionBrowserController::Init(const FConcertComponentInitParams& Params)
+	{
+		ServerInstance = Params.Server;
+		Owner = Params.WindowController;
+		
+		FGlobalTabmanager::Get()->RegisterTabSpawner(
+				ConcertServerTabs::GetSessionBrowserTabId(),
+				FOnSpawnTab::CreateRaw(this, &FConcertServerSessionBrowserController::SpawnSessionBrowserTab)
+			)
+			.SetDisplayName(LOCTEXT("SessionBrowserTabTitleLong", "Session Browser"))
+			.SetTooltipText(LOCTEXT("SessionBrowserTooltipText", "A section to browse, start, archive, and restore server sessions."))
+			.SetIcon(FSlateIcon(FConcertServerStyle::GetStyleSetName(), TEXT("Concert.MultiUser"))
+		);
+		Params.MainStack->AddTab(ConcertServerTabs::GetSessionBrowserTabId(), ETabState::OpenedTab)
+			->SetForegroundTab(ConcertServerTabs::GetSessionBrowserTabId());
 
-TArray<FConcertServerInfo> FConcertServerSessionBrowserController::GetServers() const
-{
-	return { ServerInstance->GetConcertServer()->GetServerInfo() };
-}
+		ConcertServerEvents::OnLiveSessionCreated().AddSP(this, &FConcertServerSessionBrowserController::OnLiveSessionCreated);
+		ConcertServerEvents::OnLiveSessionDestroyed().AddSP(this, &FConcertServerSessionBrowserController::OnLiveSessionDestroyed);
+		ConcertServerEvents::OnArchivedSessionCreated().AddSP(this, &FConcertServerSessionBrowserController::OnArchivedSessionCreated);
+		ConcertServerEvents::OnArchivedSessionDestroyed().AddSP(this, &FConcertServerSessionBrowserController::OnArchivedSessionDestroyed);
+	}
 
-TArray<IConcertSessionBrowserController::FActiveSessionInfo> FConcertServerSessionBrowserController::GetActiveSessions() const
-{
-	const FConcertServerInfo& ServerInfo = ServerInstance->GetConcertServer()->GetServerInfo();
-	const TArray<TSharedPtr<IConcertServerSession>> ServerSessions = ServerInstance->GetConcertServer()->GetLiveSessions();
-	
-	TArray<FActiveSessionInfo> Result;
-	Result.Reserve(ServerSessions.Num());
-	for (const TSharedPtr<IConcertServerSession>& LiveSession : ServerSessions)
+	TArray<FConcertServerInfo> FConcertServerSessionBrowserController::GetServers() const
 	{
-		FActiveSessionInfo Info{ ServerInfo, LiveSession->GetSessionInfo(), LiveSession->GetSessionClients() };
-		Result.Add(Info);
+		return { ServerInstance->GetConcertServer()->GetServerInfo() };
 	}
 
-	return Result;
-}
-
-TArray<IConcertSessionBrowserController::FArchivedSessionInfo> FConcertServerSessionBrowserController::GetArchivedSessions() const
-{
-	const FConcertServerInfo& ServerInfo = ServerInstance->GetConcertServer()->GetServerInfo();
-	const TArray<FConcertSessionInfo> ConcertSessionInfos = ServerInstance->GetConcertServer()->GetArchivedSessionInfos();
-	
-	TArray<FArchivedSessionInfo> Result;
-	Result.Reserve(ConcertSessionInfos.Num());
-	for (const FConcertSessionInfo& SessionInfo : ConcertSessionInfos)
+	TArray<IConcertSessionBrowserController::FActiveSessionInfo> FConcertServerSessionBrowserController::GetActiveSessions() const
 	{
-		Result.Emplace(ServerInfo, SessionInfo);
-	}
+		const FConcertServerInfo& ServerInfo = ServerInstance->GetConcertServer()->GetServerInfo();
+		const TArray<TSharedPtr<IConcertServerSession>> ServerSessions = ServerInstance->GetConcertServer()->GetLiveSessions();
+		
+		TArray<FActiveSessionInfo> Result;
+		Result.Reserve(ServerSessions.Num());
+		for (const TSharedPtr<IConcertServerSession>& LiveSession : ServerSessions)
+		{
+			FActiveSessionInfo Info{ ServerInfo, LiveSession->GetSessionInfo(), LiveSession->GetSessionClients() };
+			Result.Add(Info);
+		}
 
-	return Result;
-}
+		return Result;
+	}
 
-TOptional<FConcertSessionInfo> FConcertServerSessionBrowserController::GetActiveSessionInfo(const FGuid& AdminEndpoint, const FGuid& SessionId) const
-{
-	const TSharedPtr<IConcertServerSession> LiveSession = ServerInstance->GetConcertServer()->GetLiveSession(SessionId);
-	return LiveSession ? LiveSession->GetSessionInfo() : TOptional<FConcertSessionInfo>{};
-}
+	TArray<IConcertSessionBrowserController::FArchivedSessionInfo> FConcertServerSessionBrowserController::GetArchivedSessions() const
+	{
+		const FConcertServerInfo& ServerInfo = ServerInstance->GetConcertServer()->GetServerInfo();
+		const TArray<FConcertSessionInfo> ConcertSessionInfos = ServerInstance->GetConcertServer()->GetArchivedSessionInfos();
+		
+		TArray<FArchivedSessionInfo> Result;
+		Result.Reserve(ConcertSessionInfos.Num());
+		for (const FConcertSessionInfo& SessionInfo : ConcertSessionInfos)
+		{
+			Result.Emplace(ServerInfo, SessionInfo);
+		}
 
-TOptional<FConcertSessionInfo> FConcertServerSessionBrowserController::GetArchivedSessionInfo(const FGuid& AdminEndpoint, const FGuid& SessionId) const
-{
-	return ServerInstance->GetConcertServer()->GetArchivedSessionInfo(SessionId);
-}
+		return Result;
+	}
 
-void FConcertServerSessionBrowserController::CreateSession(const FGuid& ServerAdminEndpointId, const FString& SessionName, const FString& ProjectName)
-{
-	FConcertSessionInfo SessionInfo = ServerInstance->GetConcertServer()->CreateSessionInfo();
-	SessionInfo.SessionName = SessionName;
-	SessionInfo.Settings.Initialize();
-	if (!ProjectName.IsEmpty())
+	TOptional<FConcertSessionInfo> FConcertServerSessionBrowserController::GetActiveSessionInfo(const FGuid& AdminEndpoint, const FGuid& SessionId) const
 	{
-		SessionInfo.Settings.ProjectName = ProjectName;
+		const TSharedPtr<IConcertServerSession> LiveSession = ServerInstance->GetConcertServer()->GetLiveSession(SessionId);
+		return LiveSession ? LiveSession->GetSessionInfo() : TOptional<FConcertSessionInfo>{};
 	}
-	FConcertSessionVersionInfo VersionInfo;
-	VersionInfo.Initialize(false /* bSupportMixedBuildTypes */ );
-	SessionInfo.VersionInfos.Emplace(VersionInfo);
 
-	FAsyncTaskNotification Notification(
-		MakeAsyncNotification(FText::Format(LOCTEXT("CreateSessionFmt.InProgress", "Created Session '{0}'"), FText::FromString(SessionName)))
-		);
-	FText FailureReason = FText::GetEmpty();
-	const bool bSuccess = ServerInstance->GetConcertServer()->CreateSession(SessionInfo, FailureReason) != nullptr;
-
-	CloseAsyncNotification(
-		Notification,
-		bSuccess,
-		bSuccess
-			? FText::Format(LOCTEXT("CreateSessionFmt.Success", "Created Session '{0}'"), FText::FromString(SessionName))
-			: FText::Format(LOCTEXT("CreateSessionFmt.Failure", "Failed to create Session '{0}'"), FText::FromString(SessionName)),
-		FailureReason
-		);
-}
-
-void FConcertServerSessionBrowserController::ArchiveSession(const FGuid& ServerAdminEndpointId, const FGuid& SessionId, const FString& ArchiveName, const FConcertSessionFilter& SessionFilter)
-{
-	FAsyncTaskNotification Notification(
-		MakeAsyncNotification(FText::Format(LOCTEXT("ArchivedSessionFmt.InProgress", "Archiving new Session '{0}'"), FText::FromString(ArchiveName)))
-		);
-	
-	FText FailureReason = FText::GetEmpty();
-	const bool bSuccess = ServerInstance->GetConcertServer()->ArchiveSession(SessionId, ArchiveName, SessionFilter, FailureReason).IsValid();
-	
-	CloseAsyncNotification(
-		Notification,
-		bSuccess,
-		bSuccess
-			? FText::Format(LOCTEXT("ArchiveSessionFmt.Success", "Archived Session '{0}'"), FText::FromString(ArchiveName))
-			: FText::Format(LOCTEXT("ArchiveSessionFmt.Failure", "Failed to archive Session '{0}'"), FText::FromString(ArchiveName)),
-		FailureReason
-		);
-}
+	TOptional<FConcertSessionInfo> FConcertServerSessionBrowserController::GetArchivedSessionInfo(const FGuid& AdminEndpoint, const FGuid& SessionId) const
+	{
+		return ServerInstance->GetConcertServer()->GetArchivedSessionInfo(SessionId);
+	}
 
-void FConcertServerSessionBrowserController::RestoreSession(const FGuid& ServerAdminEndpointId, const FGuid& SessionId, const FString& RestoredName, const FConcertSessionFilter& SessionFilter)
-{
-	if (TOptional<FConcertSessionInfo> SessionInfo = GetArchivedSessionInfo(ServerAdminEndpointId, SessionId))
+	void FConcertServerSessionBrowserController::CreateSession(const FGuid& ServerAdminEndpointId, const FString& SessionName, const FString& ProjectName)
 	{
+		FConcertSessionInfo SessionInfo = ServerInstance->GetConcertServer()->CreateSessionInfo();
+		SessionInfo.SessionName = SessionName;
+		SessionInfo.Settings.Initialize();
+		if (!ProjectName.IsEmpty())
+		{
+			SessionInfo.Settings.ProjectName = ProjectName;
+		}
+		FConcertSessionVersionInfo VersionInfo;
+		VersionInfo.Initialize(false /* bSupportMixedBuildTypes */ );
+		SessionInfo.VersionInfos.Emplace(VersionInfo);
+
 		FAsyncTaskNotification Notification(
-			MakeAsyncNotification(FText::Format(LOCTEXT("RestoreSessionFmt.InProgress", "Restoring Session '{0}'"), FText::FromString(RestoredName)))
+			MakeAsyncNotification(FText::Format(LOCTEXT("CreateSessionFmt.InProgress", "Created Session '{0}'"), FText::FromString(SessionName)))
 			);
-		
 		FText FailureReason = FText::GetEmpty();
-		SessionInfo->SessionName = RestoredName;
-		const bool bSuccess = ServerInstance->GetConcertServer()->RestoreSession(SessionId, *SessionInfo, SessionFilter, FailureReason).IsValid();
-		
+		const bool bSuccess = ServerInstance->GetConcertServer()->CreateSession(SessionInfo, FailureReason) != nullptr;
+
 		CloseAsyncNotification(
 			Notification,
 			bSuccess,
 			bSuccess
-				? FText::Format(LOCTEXT("RestoreSessionFmt.Success", "Restored Session '{0}'"),  FText::FromString(RestoredName))
-				: FText::Format(LOCTEXT("RestoreSessionFmt.Failure", "Failed to restore Session '{0}'"), FText::FromString(RestoredName)),
+				? FText::Format(LOCTEXT("CreateSessionFmt.Success", "Created Session '{0}'"), FText::FromString(SessionName))
+				: FText::Format(LOCTEXT("CreateSessionFmt.Failure", "Failed to create Session '{0}'"), FText::FromString(SessionName)),
 			FailureReason
-		);
+			);
 	}
-}
 
-void FConcertServerSessionBrowserController::DeleteSessions(const FGuid& ServerAdminEndpointId, const TArray<FGuid>& SessionIds)
-{
-	// We want a special message to be displayed when only deleting single items
-	if (SessionIds.Num() == 1)
+	void FConcertServerSessionBrowserController::ArchiveSession(const FGuid& ServerAdminEndpointId, const FGuid& SessionId, const FString& ArchiveName, const FConcertSessionFilter& SessionFilter)
 	{
-		const FString SessionName = GetSessionName(ServerAdminEndpointId, SessionIds[0]).Get(SessionIds[0].ToString());
 		FAsyncTaskNotification Notification(
-			MakeAsyncNotification(FText::Format(LOCTEXT("DeleteSessionFmt.InProgress", "Deleting session '{0}'"), FText::FromString(SessionName)))
+			MakeAsyncNotification(FText::Format(LOCTEXT("ArchivedSessionFmt.InProgress", "Archiving new Session '{0}'"), FText::FromString(ArchiveName)))
 			);
 		
-		FText FailureReason;
-		const bool bSuccess = DeleteSession(ServerAdminEndpointId, SessionIds[0], FailureReason);
+		FText FailureReason = FText::GetEmpty();
+		const bool bSuccess = ServerInstance->GetConcertServer()->ArchiveSession(SessionId, ArchiveName, SessionFilter, FailureReason).IsValid();
 		
 		CloseAsyncNotification(
 			Notification,
 			bSuccess,
 			bSuccess
-				? FText::Format(LOCTEXT("DeleteSessionFmt.Success", "Deleted Session '{0}'"), FText::FromString(SessionName))
-				: FText::Format(LOCTEXT("DeleteSessionFmt.Failure", "Failed to archive Session '{0}'"), FText::FromString(SessionIds[0].ToString())),
+				? FText::Format(LOCTEXT("ArchiveSessionFmt.Success", "Archived Session '{0}'"), FText::FromString(ArchiveName))
+				: FText::Format(LOCTEXT("ArchiveSessionFmt.Failure", "Failed to archive Session '{0}'"), FText::FromString(ArchiveName)),
 			FailureReason
-		);
-		return;
+			);
 	}
 
-	FAsyncTaskNotification Notification(
-		MakeAsyncNotification(FText::Format(LOCTEXT("DeleteSessionsFmt.InProgress", "Deleting {0} sessions"), SessionIds.Num()))
-		);
-	
-	using FSessionName = FString;
-	using FReasonText = FText;
-	
-	TArray<TPair<FSessionName, FReasonText>> Failures;
-	int32 NumberSuccessful = 0;
-	for (const FGuid& SessionId : SessionIds)
+	void FConcertServerSessionBrowserController::RestoreSession(const FGuid& ServerAdminEndpointId, const FGuid& SessionId, const FString& RestoredName, const FConcertSessionFilter& SessionFilter)
 	{
-		FText FailureReason;
-		if (DeleteSession(ServerAdminEndpointId, SessionId, FailureReason))
+		if (TOptional<FConcertSessionInfo> SessionInfo = GetArchivedSessionInfo(ServerAdminEndpointId, SessionId))
 		{
-			++NumberSuccessful;
+			FAsyncTaskNotification Notification(
+				MakeAsyncNotification(FText::Format(LOCTEXT("RestoreSessionFmt.InProgress", "Restoring Session '{0}'"), FText::FromString(RestoredName)))
+				);
+			
+			FText FailureReason = FText::GetEmpty();
+			SessionInfo->SessionName = RestoredName;
+			const bool bSuccess = ServerInstance->GetConcertServer()->RestoreSession(SessionId, *SessionInfo, SessionFilter, FailureReason).IsValid();
+			
+			CloseAsyncNotification(
+				Notification,
+				bSuccess,
+				bSuccess
+					? FText::Format(LOCTEXT("RestoreSessionFmt.Success", "Restored Session '{0}'"),  FText::FromString(RestoredName))
+					: FText::Format(LOCTEXT("RestoreSessionFmt.Failure", "Failed to restore Session '{0}'"), FText::FromString(RestoredName)),
+				FailureReason
+			);
 		}
-		else 
+	}
+
+	void FConcertServerSessionBrowserController::DeleteSessions(const FGuid& ServerAdminEndpointId, const TArray<FGuid>& SessionIds)
+	{
+		// We want a special message to be displayed when only deleting single items
+		if (SessionIds.Num() == 1)
 		{
-			const TOptional<FString> SessionName = GetSessionName(ServerAdminEndpointId, SessionId);
-			Failures.Add({ SessionName.Get(SessionId.ToString()), FailureReason });
+			const FString SessionName = GetSessionName(ServerAdminEndpointId, SessionIds[0]).Get(SessionIds[0].ToString());
+			FAsyncTaskNotification Notification(
+				MakeAsyncNotification(FText::Format(LOCTEXT("DeleteSessionFmt.InProgress", "Deleting session '{0}'"), FText::FromString(SessionName)))
+				);
+			
+			FText FailureReason;
+			const bool bSuccess = DeleteSession(ServerAdminEndpointId, SessionIds[0], FailureReason);
+			
+			CloseAsyncNotification(
+				Notification,
+				bSuccess,
+				bSuccess
+					? FText::Format(LOCTEXT("DeleteSessionFmt.Success", "Deleted Session '{0}'"), FText::FromString(SessionName))
+					: FText::Format(LOCTEXT("DeleteSessionFmt.Failure", "Failed to archive Session '{0}'"), FText::FromString(SessionIds[0].ToString())),
+				FailureReason
+			);
+			return;
+		}
+
+		FAsyncTaskNotification Notification(
+			MakeAsyncNotification(FText::Format(LOCTEXT("DeleteSessionsFmt.InProgress", "Deleting {0} sessions"), SessionIds.Num()))
+			);
+		
+		using FSessionName = FString;
+		using FReasonText = FText;
+		
+		TArray<TPair<FSessionName, FReasonText>> Failures;
+		int32 NumberSuccessful = 0;
+		for (const FGuid& SessionId : SessionIds)
+		{
+			FText FailureReason;
+			if (DeleteSession(ServerAdminEndpointId, SessionId, FailureReason))
+			{
+				++NumberSuccessful;
+			}
+			else 
+			{
+				const TOptional<FString> SessionName = GetSessionName(ServerAdminEndpointId, SessionId);
+				Failures.Add({ SessionName.Get(SessionId.ToString()), FailureReason });
+			}
+		}
+		
+		if (Failures.Num() == 0)
+		{
+			CloseAsyncNotification(
+				Notification,
+				true,
+				FText::Format(LOCTEXT("DeleteSessionsFmt.Success", "Deleted {0} Sessions"), NumberSuccessful)
+			);
+		}
+		else
+		{
+			// Instead of a notification for each failure, show one big one so the user does not have to manually close each (cuz that's annoying...)
+			TStringBuilder<512> SubText;
+			for (const TPair<FSessionName, FReasonText>& Failure : Failures)
+			{
+				SubText << Failure.Key << TEXT(": ") << Failure.Value.ToString() << TEXT("\n");
+			}
+			CloseAsyncNotification(
+				Notification,
+				false,
+				FText::Format(LOCTEXT("DeletedSessionsFmt.Failure", "Deleted {0} of {1} Sessions."), NumberSuccessful, SessionIds.Num()),
+				FText::FromString(SubText.ToString())
+			);
 		}
 	}
-	
-	if (Failures.Num() == 0)
+
+	TSharedRef<SDockTab> FConcertServerSessionBrowserController::SpawnSessionBrowserTab(const FSpawnTabArgs& Args)
 	{
-		CloseAsyncNotification(
-			Notification,
-			true,
-			FText::Format(LOCTEXT("DeleteSessionsFmt.Success", "Deleted {0} Sessions"), NumberSuccessful)
-		);
+		const TSharedRef<SDockTab> DockTab = SNew(SDockTab)
+			.Label(LOCTEXT("SessionBrowserTabTitle", "Sessions"))
+			.TabRole(MajorTab)
+			.OnCanCloseTab(SDockTab::FCanCloseTab::CreateLambda([](){ return false; }))
+			[
+				SAssignNew(ConcertBrowser, SConcertServerSessionBrowser, SharedThis(this))
+					.DoubleClickLiveSession(this, &FConcertServerSessionBrowserController::OpenSession)
+					.DoubleClickArchivedSession(this, &FConcertServerSessionBrowserController::OpenSession)
+			];
+
+		return DockTab;
 	}
-	else
+
+	void FConcertServerSessionBrowserController::RefreshSessionList()
 	{
-		// Instead of a notification for each failure, show one big one so the user does not have to manually close each (cuz that's annoying...)
-		TStringBuilder<512> SubText;
-		for (const TPair<FSessionName, FReasonText>& Failure : Failures)
+		if (ConcertBrowser.IsValid())
 		{
-			SubText << Failure.Key << TEXT(": ") << Failure.Value.ToString() << TEXT("\n");
+			// 1. Accumulate several events in same tick so list is only updated once
+			// 2. Several events, e.g. OnLiveSessionCreated, are called before the session creation is visible to getter functions
+			ConcertBrowser->RequestRefreshListNextTick();
 		}
-		CloseAsyncNotification(
-			Notification,
-			false,
-			FText::Format(LOCTEXT("DeletedSessionsFmt.Failure", "Deleted {0} of {1} Sessions."), NumberSuccessful, SessionIds.Num()),
-			FText::FromString(SubText.ToString())
-		);
 	}
-}
-
-TSharedRef<SDockTab> FConcertServerSessionBrowserController::SpawnSessionBrowserTab(const FSpawnTabArgs& Args)
-{
-	const TSharedRef<SDockTab> DockTab = SNew(SDockTab)
-		.Label(LOCTEXT("SessionBrowserTabTitle", "Sessions"))
-		.TabRole(MajorTab)
-		.OnCanCloseTab(SDockTab::FCanCloseTab::CreateLambda([](){ return false; }))
-		[
-			SAssignNew(ConcertBrowser, SConcertServerSessionBrowser, SharedThis(this))
-				.DoubleClickLiveSession(this, &FConcertServerSessionBrowserController::OpenSession)
-				.DoubleClickArchivedSession(this, &FConcertServerSessionBrowserController::OpenSession)
-		];
-
-	return DockTab;
-}
 
-void FConcertServerSessionBrowserController::RefreshSessionList()
-{
-	if (ConcertBrowser.IsValid())
+	void FConcertServerSessionBrowserController::OpenSession(const TSharedPtr<FConcertSessionItem>& SessionItem)
 	{
-		// 1. Accumulate several events in same tick so list is only updated once
-		// 2. Several events, e.g. OnLiveSessionCreated, are called before the session creation is visible to getter functions
-		ConcertBrowser->RequestRefreshListNextTick();
+		Owner.Pin()->OpenSessionTab(SessionItem->SessionId);
 	}
-}
 
-void FConcertServerSessionBrowserController::OpenSession(const TSharedPtr<FConcertSessionItem>& SessionItem)
-{
-	Owner.Pin()->OpenSessionTab(SessionItem->SessionId);
-}
-
-void FConcertServerSessionBrowserController::RenameActiveSessionInternal(const FGuid& ServerAdminEndpointId, const FGuid& SessionId, const FString& NewName)
-{
-	if (const TOptional<FConcertSessionInfo> SessionInfo = GetActiveSessionInfo(ServerAdminEndpointId, SessionId))
+	void FConcertServerSessionBrowserController::RenameActiveSessionInternal(const FGuid& ServerAdminEndpointId, const FGuid& SessionId, const FString& NewName)
 	{
-		RenameSessionInternal(SessionId, NewName, *SessionInfo);
+		if (const TOptional<FConcertSessionInfo> SessionInfo = GetActiveSessionInfo(ServerAdminEndpointId, SessionId))
+		{
+			RenameSessionInternal(SessionId, NewName, *SessionInfo);
+		}
 	}
-}
 
-void FConcertServerSessionBrowserController::RenameArchivedSessionInternal(const FGuid& ServerAdminEndpointId, const FGuid& SessionId, const FString& NewName)
-{
-	if (const TOptional<FConcertSessionInfo> SessionInfo = GetArchivedSessionInfo(ServerAdminEndpointId, SessionId))
+	void FConcertServerSessionBrowserController::RenameArchivedSessionInternal(const FGuid& ServerAdminEndpointId, const FGuid& SessionId, const FString& NewName)
 	{
-		RenameSessionInternal(SessionId, NewName, *SessionInfo);
+		if (const TOptional<FConcertSessionInfo> SessionInfo = GetArchivedSessionInfo(ServerAdminEndpointId, SessionId))
+		{
+			RenameSessionInternal(SessionId, NewName, *SessionInfo);
+		}
 	}
-}
 
-void FConcertServerSessionBrowserController::RenameSessionInternal(const FGuid& SessionId, const FString& NewName, const FConcertSessionInfo& SessionInfo)
-{
-	FAsyncTaskNotification Notification(
-		MakeAsyncNotification(FText::Format(LOCTEXT("RenameSessionFmt.InProgress", "Renaming Session {0} "), FText::FromString(SessionInfo.SessionName)))
-		);
-	
-	FText FailureReason = FText::GetEmpty();
-	const bool bSuccess = ServerInstance->GetConcertServer()->RenameSession(SessionId, NewName, FailureReason);
+	void FConcertServerSessionBrowserController::RenameSessionInternal(const FGuid& SessionId, const FString& NewName, const FConcertSessionInfo& SessionInfo)
+	{
+		FAsyncTaskNotification Notification(
+			MakeAsyncNotification(FText::Format(LOCTEXT("RenameSessionFmt.InProgress", "Renaming Session {0} "), FText::FromString(SessionInfo.SessionName)))
+			);
+		
+		FText FailureReason = FText::GetEmpty();
+		const bool bSuccess = ServerInstance->GetConcertServer()->RenameSession(SessionId, NewName, FailureReason);
 
-	CloseAsyncNotification(
-			Notification,
-			bSuccess,
-			bSuccess
-				? FText::Format(LOCTEXT("RenameSessionFmt.Success", "Rename Session '{0}' as '{1}'"), FText::FromString(SessionInfo.SessionName), FText::FromString(NewName))
-				: FText::Format(LOCTEXT("RenameSessionFmt.Failure", "Failed to rename Session '{0}' as '{1}'"), FText::FromString(SessionInfo.SessionName), FText::FromString(NewName)),
-			FailureReason
-		);
-}
+		CloseAsyncNotification(
+				Notification,
+				bSuccess,
+				bSuccess
+					? FText::Format(LOCTEXT("RenameSessionFmt.Success", "Rename Session '{0}' as '{1}'"), FText::FromString(SessionInfo.SessionName), FText::FromString(NewName))
+					: FText::Format(LOCTEXT("RenameSessionFmt.Failure", "Failed to rename Session '{0}' as '{1}'"), FText::FromString(SessionInfo.SessionName), FText::FromString(NewName)),
+				FailureReason
+			);
+	}
 
-bool FConcertServerSessionBrowserController::DeleteSession(const FGuid& ServerAdminEndpointId, const FGuid& SessionId, FText& FailureReason)
-{
-	const TOptional<FConcertSessionInfo> ArchivedInfo = GetArchivedSessionInfo(ServerAdminEndpointId, SessionId);
-	const TOptional<FConcertSessionInfo> LiveInfo = GetActiveSessionInfo(ServerAdminEndpointId, SessionId);
-	if (const TOptional<FString> SessionName = GetSessionName(ServerAdminEndpointId, SessionId))
+	bool FConcertServerSessionBrowserController::DeleteSession(const FGuid& ServerAdminEndpointId, const FGuid& SessionId, FText& FailureReason)
 	{
-		return ServerInstance->GetConcertServer()->DestroySession(SessionId, FailureReason);
+		const TOptional<FConcertSessionInfo> ArchivedInfo = GetArchivedSessionInfo(ServerAdminEndpointId, SessionId);
+		const TOptional<FConcertSessionInfo> LiveInfo = GetActiveSessionInfo(ServerAdminEndpointId, SessionId);
+		if (const TOptional<FString> SessionName = GetSessionName(ServerAdminEndpointId, SessionId))
+		{
+			return ServerInstance->GetConcertServer()->DestroySession(SessionId, FailureReason);
+		}
+		return false;
 	}
-	return false;
-}
 
-TOptional<FString> FConcertServerSessionBrowserController::GetSessionName(const FGuid& ServerAdminEndpointId, const FGuid& SessionId) const
-{
-	const TOptional<FConcertSessionInfo> ArchivedInfo = GetArchivedSessionInfo(ServerAdminEndpointId, SessionId);
-	const TOptional<FConcertSessionInfo> LiveInfo = GetActiveSessionInfo(ServerAdminEndpointId, SessionId);
-	const TOptional<FString> SessionName = ArchivedInfo ? ArchivedInfo->SessionName : LiveInfo ? LiveInfo->SessionName : TOptional<FString>{};
-	return SessionName;
-}
-
-FAsyncTaskNotificationConfig FConcertServerSessionBrowserController::MakeAsyncNotification(const FText& Title, const FText& Details)
-{
-	FAsyncTaskNotificationConfig NotificationConfig;
-	NotificationConfig.bIsHeadless = false;
-	NotificationConfig.bKeepOpenOnFailure = true;
-	NotificationConfig.TitleText = Title;
-	return NotificationConfig;
-}
+	TOptional<FString> FConcertServerSessionBrowserController::GetSessionName(const FGuid& ServerAdminEndpointId, const FGuid& SessionId) const
+	{
+		const TOptional<FConcertSessionInfo> ArchivedInfo = GetArchivedSessionInfo(ServerAdminEndpointId, SessionId);
+		const TOptional<FConcertSessionInfo> LiveInfo = GetActiveSessionInfo(ServerAdminEndpointId, SessionId);
+		const TOptional<FString> SessionName = ArchivedInfo ? ArchivedInfo->SessionName : LiveInfo ? LiveInfo->SessionName : TOptional<FString>{};
+		return SessionName;
+	}
 
-void FConcertServerSessionBrowserController::CloseAsyncNotification(FAsyncTaskNotification& Notification, bool bSuccess, const FText& Title, const FText& Details)
-{
-	Notification.SetComplete(Title, Details, bSuccess);
+	FAsyncTaskNotificationConfig FConcertServerSessionBrowserController::MakeAsyncNotification(const FText& Title, const FText& Details)
+	{
+		FAsyncTaskNotificationConfig NotificationConfig;
+		NotificationConfig.bIsHeadless = false;
+		NotificationConfig.bKeepOpenOnFailure = true;
+		NotificationConfig.TitleText = Title;
+		return NotificationConfig;
+	}
 
-	if (bSuccess)
+	void FConcertServerSessionBrowserController::CloseAsyncNotification(FAsyncTaskNotification& Notification, bool bSuccess, const FText& Title, const FText& Details)
 	{
-		ConcertBrowser->RequestRefreshListNextTick();
+		Notification.SetComplete(Title, Details, bSuccess);
+
+		if (bSuccess)
+		{
+			ConcertBrowser->RequestRefreshListNextTick();
+		}
 	}
 }
 
diff --git a/Engine/Plugins/Developer/Concert/ConcertApp/MultiUserServer/Source/MultiUserServer/Private/Widgets/Browser/ConcertServerSessionBrowserController.h b/Engine/Plugins/Developer/Concert/ConcertApp/MultiUserServer/Source/MultiUserServer/Private/Widgets/Browser/ConcertServerSessionBrowserController.h
--- a/Engine/Plugins/Developer/Concert/ConcertApp/MultiUserServer/Source/MultiUserServer/Private/Widgets/Browser/ConcertServerSessionBrowserController.h
+++ b/Engine/Plugins/Developer/Concert/ConcertApp/MultiUserServer/Source/MultiUserServer/Private/Widgets/Browser/ConcertServerSessionBrowserController.h
@@ -12,75 +12,79 @@ class IConcertServer;
 class FConcertSessionItem;
 class FSpawnTabArgs;
 class SDockTab;
-class SConcertServerSessionBrowser;
-
-/** Implements controller in model-view-controller pattern. */
-class FConcertServerSessionBrowserController
-	:
-	public TSharedFromThis<FConcertServerSessionBrowserController>,
-	public IConcertComponent,
-	public IConcertSessionBrowserController
-{
-public:
-
-	int32 GetNumConnectedClients(const FGuid& SessionId) const;
-
-	~FConcertServerSessionBrowserController();
-
-	//~ Begin IConcertComponent Interface
-	virtual void Init(const FConcertComponentInitParams& Params) override;
-	//~ End IConcertComponent Interface
-
-	//~ Begin IConcertSessionBrowserController Interface
-	virtual TArray<FConcertServerInfo> GetServers() const override;
-	virtual TArray<FActiveSessionInfo> GetActiveSessions() const override;
-	virtual TArray<FArchivedSessionInfo> GetArchivedSessions() const override;
-	virtual TOptional<FConcertSessionInfo> GetActiveSessionInfo(const FGuid& AdminEndpoint, const FGuid& SessionId) const override;
-	virtual TOptional<FConcertSessionInfo> GetArchivedSessionInfo(const FGuid& AdminEndpoint, const FGuid& SessionId) const override;
-
-	virtual void CreateSession(const FGuid& ServerAdminEndpointId, const FString& SessionName, const FString& ProjectName) override;
-	virtual void ArchiveSession(const FGuid& ServerAdminEndpointId, const FGuid& SessionId, const FString& ArchiveName, const FConcertSessionFilter& SessionFilter) override;
-	virtual void RestoreSession(const FGuid& ServerAdminEndpointId, const FGuid& SessionId, const FString& RestoredName, const FConcertSessionFilter& SessionFilter) override;
-	virtual void RenameActiveSession(const FGuid& ServerAdminEndpointId, const FGuid& SessionId, const FString& NewName) override { RenameActiveSessionInternal(ServerAdminEndpointId, SessionId, NewName); }
-	virtual void RenameArchivedSession(const FGuid& ServerAdminEndpointId, const FGuid& SessionId, const FString& NewName) override { RenameArchivedSessionInternal(ServerAdminEndpointId, SessionId, NewName); }
-	virtual void DeleteSessions(const FGuid& ServerAdminEndpointId, const TArray<FGuid>& SessionIds) override;
-
-	// The server operator always has permission for these actions:
-	virtual bool CanRenameActiveSession(const FGuid& ServerAdminEndpointId, const FGuid& SessionId) const override { return true; }
-	virtual bool CanRenameArchivedSession(const FGuid& ServerAdminEndpointId, const FGuid& SessionId) const override { return true; }
-	virtual bool CanDeleteActiveSession(const FGuid& ServerAdminEndpointId, const FGuid& SessionId) const override { return true; }
-	virtual bool CanDeleteArchivedSession(const FGuid& ServerAdminEndpointId, const FGuid& SessionId) const override { return true; }
-	//~ Begin IConcertSessionBrowserController Interface
-
-private:
 
-	/** Server instance we're representing */
-	TSharedPtr<IConcertSyncServer> ServerInstance;
-	/** Used to open selected sessions */
-	TWeakPtr<FConcertServerWindowController> Owner;
-
-	TSharedPtr<SConcertServerSessionBrowser> ConcertBrowser;
-
-	TSharedRef<SDockTab> SpawnSessionBrowserTab(const FSpawnTabArgs& Args);
-
-	// Update view when session list changes
-	void OnLiveSessionCreated(bool, const IConcertServer&, TSharedRef<IConcertServerSession>) { RefreshSessionList(); }
-	void OnLiveSessionDestroyed(const IConcertServer&, TSharedRef<IConcertServerSession>) { RefreshSessionList(); }
-	void OnArchivedSessionCreated(bool, const IConcertServer&, const FString&, const FConcertSessionInfo&) { RefreshSessionList(); }
-	void OnArchivedSessionDestroyed(const IConcertServer&, const FGuid&) { RefreshSessionList(); }
-
-	void RefreshSessionList();
-
-	// Session actions
-	void OpenSession(const TSharedPtr<FConcertSessionItem>& SessionItem);
-	void RenameActiveSessionInternal(const FGuid& ServerAdminEndpointId, const FGuid& SessionId, const FString& NewName);
-	void RenameArchivedSessionInternal(const FGuid& ServerAdminEndpointId, const FGuid& SessionId, const FString& NewName);
-	void RenameSessionInternal(const FGuid& SessionId, const FString& NewName, const FConcertSessionInfo& SessionInfo);
-	bool DeleteSession(const FGuid& ServerAdminEndpointId, const FGuid& SessionId, FText& FailureReason);
-
-	TOptional<FString> GetSessionName(const FGuid& ServerAdminEndpointId, const FGuid& SessionId) const;
-
-	FAsyncTaskNotificationConfig MakeAsyncNotification(const FText& Title, const FText& Details = FText::GetEmpty());
-	void CloseAsyncNotification(FAsyncTaskNotification& Notification, bool bSuccess, const FText& Title, const FText& Details = FText::GetEmpty());
-};
+namespace UE::MultiUserServer
+{
+	class SConcertServerSessionBrowser;
+	
+	/** Implements controller in model-view-controller pattern. */
+	class FConcertServerSessionBrowserController
+		:
+		public TSharedFromThis<FConcertServerSessionBrowserController>,
+		public IConcertComponent,
+		public IConcertSessionBrowserController
+	{
+	public:
+
+		int32 GetNumConnectedClients(const FGuid& SessionId) const;
+
+		~FConcertServerSessionBrowserController();
+
+		//~ Begin IConcertComponent Interface
+		virtual void Init(const FConcertComponentInitParams& Params) override;
+		//~ End IConcertComponent Interface
+
+		//~ Begin IConcertSessionBrowserController Interface
+		virtual TArray<FConcertServerInfo> GetServers() const override;
+		virtual TArray<FActiveSessionInfo> GetActiveSessions() const override;
+		virtual TArray<FArchivedSessionInfo> GetArchivedSessions() const override;
+		virtual TOptional<FConcertSessionInfo> GetActiveSessionInfo(const FGuid& AdminEndpoint, const FGuid& SessionId) const override;
+		virtual TOptional<FConcertSessionInfo> GetArchivedSessionInfo(const FGuid& AdminEndpoint, const FGuid& SessionId) const override;
+
+		virtual void CreateSession(const FGuid& ServerAdminEndpointId, const FString& SessionName, const FString& ProjectName) override;
+		virtual void ArchiveSession(const FGuid& ServerAdminEndpointId, const FGuid& SessionId, const FString& ArchiveName, const FConcertSessionFilter& SessionFilter) override;
+		virtual void RestoreSession(const FGuid& ServerAdminEndpointId, const FGuid& SessionId, const FString& RestoredName, const FConcertSessionFilter& SessionFilter) override;
+		virtual void RenameActiveSession(const FGuid& ServerAdminEndpointId, const FGuid& SessionId, const FString& NewName) override { RenameActiveSessionInternal(ServerAdminEndpointId, SessionId, NewName); }
+		virtual void RenameArchivedSession(const FGuid& ServerAdminEndpointId, const FGuid& SessionId, const FString& NewName) override { RenameArchivedSessionInternal(ServerAdminEndpointId, SessionId, NewName); }
+		virtual void DeleteSessions(const FGuid& ServerAdminEndpointId, const TArray<FGuid>& SessionIds) override;
+
+		// The server operator always has permission for these actions:
+		virtual bool CanRenameActiveSession(const FGuid& ServerAdminEndpointId, const FGuid& SessionId) const override { return true; }
+		virtual bool CanRenameArchivedSession(const FGuid& ServerAdminEndpointId, const FGuid& SessionId) const override { return true; }
+		virtual bool CanDeleteActiveSession(const FGuid& ServerAdminEndpointId, const FGuid& SessionId) const override { return true; }
+		virtual bool CanDeleteArchivedSession(const FGuid& ServerAdminEndpointId, const FGuid& SessionId) const override { return true; }
+		//~ Begin IConcertSessionBrowserController Interface
+
+	private:
+
+		/** Server instance we're representing */
+		TSharedPtr<IConcertSyncServer> ServerInstance;
+		/** Used to open selected sessions */
+		TWeakPtr<FConcertServerWindowController> Owner;
+
+		TSharedPtr<SConcertServerSessionBrowser> ConcertBrowser;
+
+		TSharedRef<SDockTab> SpawnSessionBrowserTab(const FSpawnTabArgs& Args);
+
+		// Update view when session list changes
+		void OnLiveSessionCreated(bool, const IConcertServer&, TSharedRef<IConcertServerSession>) { RefreshSessionList(); }
+		void OnLiveSessionDestroyed(const IConcertServer&, TSharedRef<IConcertServerSession>) { RefreshSessionList(); }
+		void OnArchivedSessionCreated(bool, const IConcertServer&, const FString&, const FConcertSessionInfo&) { RefreshSessionList(); }
+		void OnArchivedSessionDestroyed(const IConcertServer&, const FGuid&) { RefreshSessionList(); }
+
+		void RefreshSessionList();
+
+		// Session actions
+		void OpenSession(const TSharedPtr<FConcertSessionItem>& SessionItem);
+		void RenameActiveSessionInternal(const FGuid& ServerAdminEndpointId, const FGuid& SessionId, const FString& NewName);
+		void RenameArchivedSessionInternal(const FGuid& ServerAdminEndpointId, const FGuid& SessionId, const FString& NewName);
+		void RenameSessionInternal(const FGuid& SessionId, const FString& NewName, const FConcertSessionInfo& SessionInfo);
+		bool DeleteSession(const FGuid& ServerAdminEndpointId, const FGuid& SessionId, FText& FailureReason);
+
+		TOptional<FString> GetSessionName(const FGuid& ServerAdminEndpointId, const FGuid& SessionId) const;
+
+		FAsyncTaskNotificationConfig MakeAsyncNotification(const FText& Title, const FText& Details = FText::GetEmpty());
+		void CloseAsyncNotification(FAsyncTaskNotification& Notification, bool bSuccess, const FText& Title, const FText& Details = FText::GetEmpty());
+	};
+}
 
diff --git a/Engine/Plugins/Developer/Concert/ConcertApp/MultiUserServer/Source/MultiUserServer/Private/Widgets/Browser/SConcertServerSessionBrowser.cpp b/Engine/Plugins/Developer/Concert/ConcertApp/MultiUserServer/Source/MultiUserServer/Private/Widgets/Browser/SConcertServerSessionBrowser.cpp
--- a/Engine/Plugins/Developer/Concert/ConcertApp/MultiUserServer/Source/MultiUserServer/Private/Widgets/Browser/SConcertServerSessionBrowser.cpp
+++ b/Engine/Plugins/Developer/Concert/ConcertApp/MultiUserServer/Source/MultiUserServer/Private/Widgets/Browser/SConcertServerSessionBrowser.cpp
@@ -13,95 +13,98 @@
 
 #include "Dialog/SMessageDialog.h"
 
-#define LOCTEXT_NAMESPACE "UnrealMultiUserUI"
+#define LOCTEXT_NAMESPACE "UnrealMultiUserUI.SConcertServerSessionBrowser"
 
-void SConcertServerSessionBrowser::Construct(const FArguments& InArgs, TSharedRef<FConcertServerSessionBrowserController> InController)
+namespace UE::MultiUserServer
 {
-	Controller = InController;
-	SConcertTabViewBase::Construct(
-		SConcertTabViewBase::FArguments()
-		.Content()
-		[
-			MakeSessionTableView(InArgs)
-		],
-		ConcertServerTabs::GetSessionBrowserTabId()
-		);
-}
+	void SConcertServerSessionBrowser::Construct(const FArguments& InArgs, TSharedRef<FConcertServerSessionBrowserController> InController)
+	{
+		Controller = InController;
+		SConcertTabViewBase::Construct(
+			SConcertTabViewBase::FArguments()
+			.Content()
+			[
+				MakeSessionTableView(InArgs)
+			],
+			ConcertServerTabs::GetSessionBrowserTabId()
+			);
+	}
 
-TSharedRef<SWidget> SConcertServerSessionBrowser::MakeSessionTableView(const FArguments& InArgs)
-{
-	SearchText = MakeShared<FText>();
-	return SAssignNew(SessionBrowser, SConcertSessionBrowser, Controller.Pin().ToSharedRef(), SearchText)
-		.OnLiveSessionDoubleClicked(InArgs._DoubleClickLiveSession)
-		.OnArchivedSessionDoubleClicked(InArgs._DoubleClickArchivedSession)
-		.PostRequestedDeleteSession(this, &SConcertServerSessionBrowser::RequestDeleteSession)
-		// Pretend a modal dialog said no - RequestDeleteSession will show non-modal dialog
-		.AskUserToDeleteSessions_Lambda([](auto) { return false; })
-		.ColumnVisibilitySnapshot(UMultiUserServerColumnVisibilitySettings::GetSettings()->GetSessionBrowserColumnVisibility())
-		.SaveColumnVisibilitySnapshot_Lambda([](const FColumnVisibilitySnapshot& Snapshot)
-		{
-			UMultiUserServerColumnVisibilitySettings::GetSettings()->SetSessionBrowserColumnVisibility(Snapshot);
-		});
-}
+	TSharedRef<SWidget> SConcertServerSessionBrowser::MakeSessionTableView(const FArguments& InArgs)
+	{
+		SearchText = MakeShared<FText>();
+		return SAssignNew(SessionBrowser, SConcertSessionBrowser, Controller.Pin().ToSharedRef(), SearchText)
+			.OnLiveSessionDoubleClicked(InArgs._DoubleClickLiveSession)
+			.OnArchivedSessionDoubleClicked(InArgs._DoubleClickArchivedSession)
+			.PostRequestedDeleteSession(this, &SConcertServerSessionBrowser::RequestDeleteSession)
+			// Pretend a modal dialog said no - RequestDeleteSession will show non-modal dialog
+			.AskUserToDeleteSessions_Lambda([](auto) { return false; })
+			.ColumnVisibilitySnapshot(UMultiUserServerColumnVisibilitySettings::GetSettings()->GetSessionBrowserColumnVisibility())
+			.SaveColumnVisibilitySnapshot_Lambda([](const FColumnVisibilitySnapshot& Snapshot)
+			{
+				UMultiUserServerColumnVisibilitySettings::GetSettings()->SetSessionBrowserColumnVisibility(Snapshot);
+			});
+	}
 
-void SConcertServerSessionBrowser::RequestDeleteSession(const TArray<TSharedPtr<FConcertSessionItem>>& SessionItems) 
-{
-	const FText Message = [this, &SessionItems]()
+	void SConcertServerSessionBrowser::RequestDeleteSession(const TArray<TSharedPtr<FConcertSessionItem>>& SessionItems) 
 	{
-		if (SessionItems.Num() > 1)
-		{
-			return FText::Format(
-			LOCTEXT("DeletedMultipleDescription", "Deleting a session will force all connected clients to disconnect and all associated data to be removed.\n\nDelete {0} sessions?"),
-				SessionItems.Num()
-			);
-		}
-			
-		switch (SessionItems[0]->Type)
+		const FText Message = [this, &SessionItems]()
 		{
-		case FConcertSessionItem::EType::ActiveSession:
-			return FText::Format(
-				LOCTEXT("DeletedActiveDescription", "There {0}|plural(one=is,other=are) {0} connected {0}|plural(one=client,other=clients) in the current session.\nDeleting a session will force all connected clients to disconnect.\n\nDelete {1}?"),
-				Controller.Pin()->GetNumConnectedClients(SessionItems[0]->SessionId),
-				FText::FromString(SessionItems[0]->SessionName)
-			);
-		case FConcertSessionItem::EType::ArchivedSession:
-			return FText::Format(
-				LOCTEXT("DeleteArchivedDescription", "Deleting a session will cause all associated data to be removed.\n\nDelete {0}?"),
-				FText::FromString(SessionItems[0]->SessionName)
+			if (SessionItems.Num() > 1)
+			{
+				return FText::Format(
+				LOCTEXT("DeletedMultipleDescription", "Deleting a session will force all connected clients to disconnect and all associated data to be removed.\n\nDelete {0} sessions?"),
+					SessionItems.Num()
 				);
-		default:
-			checkNoEntry();
-			return FText::GetEmpty();
-		}
-	}();
-	DeleteSessionsWithFakeModalQuestion(Message, SessionItems);
-}
+			}
+				
+			switch (SessionItems[0]->Type)
+			{
+			case FConcertSessionItem::EType::ActiveSession:
+				return FText::Format(
+					LOCTEXT("DeletedActiveDescription", "There {0}|plural(one=is,other=are) {0} connected {0}|plural(one=client,other=clients) in the current session.\nDeleting a session will force all connected clients to disconnect.\n\nDelete {1}?"),
+					Controller.Pin()->GetNumConnectedClients(SessionItems[0]->SessionId),
+					FText::FromString(SessionItems[0]->SessionName)
+				);
+			case FConcertSessionItem::EType::ArchivedSession:
+				return FText::Format(
+					LOCTEXT("DeleteArchivedDescription", "Deleting a session will cause all associated data to be removed.\n\nDelete {0}?"),
+					FText::FromString(SessionItems[0]->SessionName)
+					);
+			default:
+				checkNoEntry();
+				return FText::GetEmpty();
+			}
+		}();
+		DeleteSessionsWithFakeModalQuestion(Message, SessionItems);
+	}
 
-void SConcertServerSessionBrowser::DeleteSessionsWithFakeModalQuestion(const FText& Message, const TArray<TSharedPtr<FConcertSessionItem>>& SessionItems)
-{
-	auto DeleteArchived = [WeakController = TWeakPtr<FConcertServerSessionBrowserController>(Controller), SessionItems]()
+	void SConcertServerSessionBrowser::DeleteSessionsWithFakeModalQuestion(const FText& Message, const TArray<TSharedPtr<FConcertSessionItem>>& SessionItems)
 	{
-		if (const TSharedPtr<FConcertServerSessionBrowserController> PinnedController = WeakController.Pin())
+		auto DeleteArchived = [WeakController = TWeakPtr<FConcertServerSessionBrowserController>(Controller), SessionItems]()
 		{
-			ConcertBrowserUtils::RequestItemDeletion(*PinnedController.Get(), SessionItems);
-		}
-	};
-	const TSharedRef<SMessageDialog> Dialog = SNew(SMessageDialog)
-		.Title(LOCTEXT("DisconnectUsersTitle", "Delete session?"))
-		.Icon(FAppStyle::Get().GetBrush("Icons.WarningWithColor.Large"))
-		.Message(Message)
-		.UseScrollBox(false)
-		.Buttons({
-			SMessageDialog::FButton(LOCTEXT("DeleteArchivedButton", "Delete"))
-				.SetOnClicked(FSimpleDelegate::CreateLambda(DeleteArchived)),
-			SMessageDialog::FButton(LOCTEXT("CancelButton", "Cancel"))
-				.SetPrimary(true)
-				.SetFocus()
-		});
+			if (const TSharedPtr<FConcertServerSessionBrowserController> PinnedController = WeakController.Pin())
+			{
+				ConcertBrowserUtils::RequestItemDeletion(*PinnedController.Get(), SessionItems);
+			}
+		};
+		const TSharedRef<SMessageDialog> Dialog = SNew(SMessageDialog)
+			.Title(LOCTEXT("DisconnectUsersTitle", "Delete session?"))
+			.Icon(FAppStyle::Get().GetBrush("Icons.WarningWithColor.Large"))
+			.Message(Message)
+			.UseScrollBox(false)
+			.Buttons({
+				SMessageDialog::FButton(LOCTEXT("DeleteArchivedButton", "Delete"))
+					.SetOnClicked(FSimpleDelegate::CreateLambda(DeleteArchived)),
+				SMessageDialog::FButton(LOCTEXT("CancelButton", "Cancel"))
+					.SetPrimary(true)
+					.SetFocus()
+			});
 
-	UE::MultiUserServer::FConcertServerUIModule::Get()
-		.GetModalWindowManager()
-		->ShowFakeModalWindow(Dialog);
+		UE::MultiUserServer::FConcertServerUIModule::Get()
+			.GetModalWindowManager()
+			->ShowFakeModalWindow(Dialog);
+	}
 }
 
 #undef LOCTEXT_NAMESPACE
diff --git a/Engine/Plugins/Developer/Concert/ConcertApp/MultiUserServer/Source/MultiUserServer/Private/Widgets/Browser/SConcertServerSessionBrowser.h b/Engine/Plugins/Developer/Concert/ConcertApp/MultiUserServer/Source/MultiUserServer/Private/Widgets/Browser/SConcertServerSessionBrowser.h
--- a/Engine/Plugins/Developer/Concert/ConcertApp/MultiUserServer/Source/MultiUserServer/Private/Widgets/Browser/SConcertServerSessionBrowser.h
+++ b/Engine/Plugins/Developer/Concert/ConcertApp/MultiUserServer/Source/MultiUserServer/Private/Widgets/Browser/SConcertServerSessionBrowser.h
@@ -7,45 +7,49 @@
 #include "Widgets/DeclarativeSyntaxSupport.h"
 #include "Widgets/SConcertTabViewBase.h"
 
-class FConcertServerSessionBrowserController;
 class SConcertSessionBrowser;
 
-/** Shows a list of server sessions */
-class SConcertServerSessionBrowser : public SConcertTabViewBase
+namespace UE::MultiUserServer
 {
-public:
+	class FConcertServerSessionBrowserController;
+
+	/** Shows a list of server sessions */
+	class SConcertServerSessionBrowser : public SConcertTabViewBase
+	{
+	public:
 	
-	SLATE_BEGIN_ARGS(SConcertServerSessionBrowser) { }
+		SLATE_BEGIN_ARGS(SConcertServerSessionBrowser) { }
 		SLATE_EVENT(FSessionDelegate, DoubleClickLiveSession)
 		SLATE_EVENT(FSessionDelegate, DoubleClickArchivedSession)
 	SLATE_END_ARGS()
 
 	void Construct(const FArguments& InArgs, TSharedRef<FConcertServerSessionBrowserController> InController);
 
-	void RequestRefreshListNextTick() { bRequestedRefresh = true; }
-	virtual void Tick(const FGeometry& AllottedGeometry, const double InCurrentTime, const float InDeltaTime) override
-	{
-		if (bRequestedRefresh)
+		void RequestRefreshListNextTick() { bRequestedRefresh = true; }
+		virtual void Tick(const FGeometry& AllottedGeometry, const double InCurrentTime, const float InDeltaTime) override
 		{
-			SessionBrowser->RefreshSessionList();
-			bRequestedRefresh = false;
-		}
+			if (bRequestedRefresh)
+			{
+				SessionBrowser->RefreshSessionList();
+				bRequestedRefresh = false;
+			}
 
-		SCompoundWidget::Tick(AllottedGeometry, InCurrentTime, InDeltaTime);
-	}
+			SCompoundWidget::Tick(AllottedGeometry, InCurrentTime, InDeltaTime);
+		}
 	
-private:
+	private:
 
-	/** We can ask the controller about information and notify it about UI events. */
-	TWeakPtr<FConcertServerSessionBrowserController> Controller;
+		/** We can ask the controller about information and notify it about UI events. */
+		TWeakPtr<FConcertServerSessionBrowserController> Controller;
 
-	bool bRequestedRefresh = false;
+		bool bRequestedRefresh = false;
 
-	TSharedPtr<FText> SearchText;
-	TSharedPtr<SConcertSessionBrowser> SessionBrowser;
+		TSharedPtr<FText> SearchText;
+		TSharedPtr<SConcertSessionBrowser> SessionBrowser;
 	
-	TSharedRef<SWidget> MakeSessionTableView(const FArguments& InArgs);
+		TSharedRef<SWidget> MakeSessionTableView(const FArguments& InArgs);
 
-	void RequestDeleteSession(const TArray<TSharedPtr<FConcertSessionItem>>& SessionItems);
-	void DeleteSessionsWithFakeModalQuestion(const FText& Message, const TArray<TSharedPtr<FConcertSessionItem>>& SessionItems);
-};
+		void RequestDeleteSession(const TArray<TSharedPtr<FConcertSessionItem>>& SessionItems);
+		void DeleteSessionsWithFakeModalQuestion(const FText& Message, const TArray<TSharedPtr<FConcertSessionItem>>& SessionItems);
+	};
+}
diff --git a/Engine/Plugins/Developer/Concert/ConcertApp/MultiUserServer/Source/MultiUserServer/Private/Widgets/Clients/ConcertClientsTabController.cpp b/Engine/Plugins/Developer/Concert/ConcertApp/MultiUserServer/Source/MultiUserServer/Private/Widgets/Clients/ConcertClientsTabController.cpp
--- a/Engine/Plugins/Developer/Concert/ConcertApp/MultiUserServer/Source/MultiUserServer/Private/Widgets/Clients/ConcertClientsTabController.cpp
+++ b/Engine/Plugins/Developer/Concert/ConcertApp/MultiUserServer/Source/MultiUserServer/Private/Widgets/Clients/ConcertClientsTabController.cpp
@@ -15,41 +15,41 @@
 namespace UE::MultiUserServer
 {
 	static constexpr uint64 GlobalLogCapacity = 5000000;
-}
-
-void FConcertClientsTabController::Init(const FConcertComponentInitParams& Params)
-{
-	LogBuffer = MakeShared<FGlobalLogSource>(UE::MultiUserServer::GlobalLogCapacity);
-	AckTracker = MakeShared<FLogAckTracker>(LogBuffer.ToSharedRef(), Params.Server->GetConcertServer());
 	
-	FGlobalTabmanager::Get()->RegisterTabSpawner(
-				ConcertServerTabs::GetClientsTabID(),
-				FOnSpawnTab::CreateRaw(this, &FConcertClientsTabController::SpawnClientsTab, Params.WindowController->GetRootWindow(), Params.Server)
-			)
-			.SetDisplayName(LOCTEXT("ClientsTabTitle", "Clients"))
-			.SetTooltipText(LOCTEXT("ClientsTooltipText", "View network statistics for connected clients.")
-		);
-	Params.MainStack->AddTab(ConcertServerTabs::GetClientsTabID(), ETabState::OpenedTab);
-}
-
-void FConcertClientsTabController::ShowConnectedClients(const FGuid& SessionId) const
-{
-	FGlobalTabmanager::Get()->TryInvokeTab(ConcertServerTabs::GetClientsTabID());
-	ClientsView->ShowConnectedClients(SessionId);
-}
-
-TSharedRef<SDockTab> FConcertClientsTabController::SpawnClientsTab(const FSpawnTabArgs& Args, TSharedPtr<SWindow> RootWindow, TSharedRef<IConcertSyncServer> Server)
-{
-	const TSharedRef<SDockTab> DockTab = SNew(SDockTab)
-		.Label(LOCTEXT("ClientsTabTitle", "Clients"))
-		.TabRole(MajorTab)
-		.CanEverClose(false);
-	DockTab->SetContent(
-		SAssignNew(ClientsView, SConcertClientsTabView, ConcertServerTabs::GetClientsTabID(), Server, LogBuffer.ToSharedRef())
-			.ConstructUnderMajorTab(DockTab)
-			.ConstructUnderWindow(RootWindow)
-		);
-	return DockTab;
+	void FConcertClientsTabController::Init(const FConcertComponentInitParams& Params)
+	{
+		LogBuffer = MakeShared<FGlobalLogSource>(UE::MultiUserServer::GlobalLogCapacity);
+		AckTracker = MakeShared<FLogAckTracker>(LogBuffer.ToSharedRef(), Params.Server->GetConcertServer());
+	
+		FGlobalTabmanager::Get()->RegisterTabSpawner(
+					ConcertServerTabs::GetClientsTabID(),
+					FOnSpawnTab::CreateRaw(this, &FConcertClientsTabController::SpawnClientsTab, Params.WindowController->GetRootWindow(), Params.Server)
+				)
+				.SetDisplayName(LOCTEXT("ClientsTabTitle", "Clients"))
+				.SetTooltipText(LOCTEXT("ClientsTooltipText", "View network statistics for connected clients.")
+			);
+		Params.MainStack->AddTab(ConcertServerTabs::GetClientsTabID(), ETabState::OpenedTab);
+	}
+
+	void FConcertClientsTabController::ShowConnectedClients(const FGuid& SessionId) const
+	{
+		FGlobalTabmanager::Get()->TryInvokeTab(ConcertServerTabs::GetClientsTabID());
+		ClientsView->ShowConnectedClients(SessionId);
+	}
+
+	TSharedRef<SDockTab> FConcertClientsTabController::SpawnClientsTab(const FSpawnTabArgs& Args, TSharedPtr<SWindow> RootWindow, TSharedRef<IConcertSyncServer> Server)
+	{
+		const TSharedRef<SDockTab> DockTab = SNew(SDockTab)
+			.Label(LOCTEXT("ClientsTabTitle", "Clients"))
+			.TabRole(MajorTab)
+			.CanEverClose(false);
+		DockTab->SetContent(
+			SAssignNew(ClientsView, SConcertClientsTabView, ConcertServerTabs::GetClientsTabID(), Server, LogBuffer.ToSharedRef())
+				.ConstructUnderMajorTab(DockTab)
+				.ConstructUnderWindow(RootWindow)
+			);
+		return DockTab;
+	}
 }
 
 #undef LOCTEXT_NAMESPACE
\ No newline at end of file
diff --git a/Engine/Plugins/Developer/Concert/ConcertApp/MultiUserServer/Source/MultiUserServer/Private/Widgets/Clients/ConcertClientsTabController.h b/Engine/Plugins/Developer/Concert/ConcertApp/MultiUserServer/Source/MultiUserServer/Private/Widgets/Clients/ConcertClientsTabController.h
--- a/Engine/Plugins/Developer/Concert/ConcertApp/MultiUserServer/Source/MultiUserServer/Private/Widgets/Clients/ConcertClientsTabController.h
+++ b/Engine/Plugins/Developer/Concert/ConcertApp/MultiUserServer/Source/MultiUserServer/Private/Widgets/Clients/ConcertClientsTabController.h
@@ -7,29 +7,33 @@
 
 class FLogAckTracker;
 class FGlobalLogSource;
-class SConcertClientsTabView;
 class SWindow;
 
-class FConcertClientsTabController : public IConcertComponent, public TSharedFromThis<FConcertClientsTabController>
+namespace UE::MultiUserServer
 {
-public:
+	class SConcertClientsTabView;
+	
+	class FConcertClientsTabController : public IConcertComponent, public TSharedFromThis<FConcertClientsTabController>
+	{
+	public:
 
-	//~ Begin IConcertComponent Interface
-	virtual void Init(const FConcertComponentInitParams& Params) override;
-	//~ End IConcertComponent Interface
+		//~ Begin IConcertComponent Interface
+		virtual void Init(const FConcertComponentInitParams& Params) override;
+		//~ End IConcertComponent Interface
 
-	/** Highlights this tab and sets the client filter such that all connected clients of the given session ID are shown.  */
-	void ShowConnectedClients(const FGuid& SessionId) const;
+		/** Highlights this tab and sets the client filter such that all connected clients of the given session ID are shown.  */
+		void ShowConnectedClients(const FGuid& SessionId) const;
 
-private:
+	private:
 
-	/** Buffers generated logs up to a limit (and overrides oldest logs when buffer is full) */
-	TSharedPtr<FGlobalLogSource> LogBuffer;
-	/** Marks logs as acked as the ACKs come in. */
-	TSharedPtr<FLogAckTracker> AckTracker;
+		/** Buffers generated logs up to a limit (and overrides oldest logs when buffer is full) */
+		TSharedPtr<FGlobalLogSource> LogBuffer;
+		/** Marks logs as acked as the ACKs come in. */
+		TSharedPtr<FLogAckTracker> AckTracker;
 	
-	/** Manages the sub-tabs */
-	TSharedPtr<SConcertClientsTabView> ClientsView;
+		/** Manages the sub-tabs */
+		TSharedPtr<SConcertClientsTabView> ClientsView;
 	
-	TSharedRef<SDockTab> SpawnClientsTab(const FSpawnTabArgs& Args, TSharedPtr<SWindow> RootWindow, TSharedRef<IConcertSyncServer>);
-};
+		TSharedRef<SDockTab> SpawnClientsTab(const FSpawnTabArgs& Args, TSharedPtr<SWindow> RootWindow, TSharedRef<IConcertSyncServer>);
+	};
+}
\ No newline at end of file
diff --git a/Engine/Plugins/Developer/Concert/ConcertApp/MultiUserServer/Source/MultiUserServer/Private/Widgets/Clients/Logging/Filter/ConcertFrontendLogFilter.h b/Engine/Plugins/Developer/Concert/ConcertApp/MultiUserServer/Source/MultiUserServer/Private/Widgets/Clients/Logging/Filter/ConcertFrontendLogFilter.h
deleted file mode 100644
diff --git a/Engine/Plugins/Developer/Concert/ConcertApp/MultiUserServer/Source/MultiUserServer/Private/Widgets/Clients/Logging/Filter/ConcertFrontendLogFilter_Ack.h b/Engine/Plugins/Developer/Concert/ConcertApp/MultiUserServer/Source/MultiUserServer/Private/Widgets/Clients/Logging/Filter/ConcertFrontendLogFilter_Ack.h
--- a/Engine/Plugins/Developer/Concert/ConcertApp/MultiUserServer/Source/MultiUserServer/Private/Widgets/Clients/Logging/Filter/ConcertFrontendLogFilter_Ack.h
+++ b/Engine/Plugins/Developer/Concert/ConcertApp/MultiUserServer/Source/MultiUserServer/Private/Widgets/Clients/Logging/Filter/ConcertFrontendLogFilter_Ack.h
@@ -6,7 +6,7 @@
 #include "ConcertFrontendLogFilter_BaseSetSelection.h"
 #include "Widgets/Clients/Logging/ConcertLogEntry.h"
 
-namespace UE::MultiUserServer::Filters
+namespace UE::MultiUserServer
 {
 	/** Allows only messages that have a selected ACK state */
 	class FConcertLogFilter_Ack : public TConcertLogFilter_BaseSetSelection<FConcertLogFilter_Ack, EConcertLogAckState>
diff --git a/Engine/Plugins/Developer/Concert/ConcertApp/MultiUserServer/Source/MultiUserServer/Private/Widgets/Clients/Logging/Filter/ConcertFrontendLogFilter_BaseSetSelection.h b/Engine/Plugins/Developer/Concert/ConcertApp/MultiUserServer/Source/MultiUserServer/Private/Widgets/Clients/Logging/Filter/ConcertFrontendLogFilter_BaseSetSelection.h
--- a/Engine/Plugins/Developer/Concert/ConcertApp/MultiUserServer/Source/MultiUserServer/Private/Widgets/Clients/Logging/Filter/ConcertFrontendLogFilter_BaseSetSelection.h
+++ b/Engine/Plugins/Developer/Concert/ConcertApp/MultiUserServer/Source/MultiUserServer/Private/Widgets/Clients/Logging/Filter/ConcertFrontendLogFilter_BaseSetSelection.h
@@ -3,8 +3,10 @@
 #pragma once
 
 #include "CoreMinimal.h"
-#include "ConcertFrontendLogFilter.h"
-#include "ConcertLogFilter.h"
+
+#include "ConcertLogFilterTypes.h"
+#include "Widgets/Util/Filter/ConcertFilter.h"
+#include "Widgets/Util/Filter/ConcertFrontendFilter.h"
 
 #include "Framework/MultiBox/MultiBoxBuilder.h"
 #include "Widgets/Input/SComboButton.h"
@@ -12,7 +14,7 @@
 
 #define LOCTEXT_NAMESPACE "UnrealMultiUserUI.BaseSetSelection"
 
-namespace UE::MultiUserServer::Filters
+namespace UE::MultiUserServer
 {
 	/**
 	 * Helper implementation for filters that allow filter based on a finte set of items.
@@ -101,7 +103,7 @@ namespace UE::MultiUserServer::Filters
 	};
 
 	template<typename TRealFilterImpl>
-	class TConcertFrontendLogFilter_BaseSetSelection : public TConcertFrontendLogFilterAggregate<TRealFilterImpl>
+	class TConcertFrontendLogFilter_BaseSetSelection : public TConcertFrontendFilterAggregate<TRealFilterImpl, const FConcertLogEntry&>
 	{
 		using TItemType = typename TRealFilterImpl::TItemType;
 	public:
diff --git a/Engine/Plugins/Developer/Concert/ConcertApp/MultiUserServer/Source/MultiUserServer/Private/Widgets/Clients/Logging/Filter/ConcertFrontendLogFilter_Client.cpp b/Engine/Plugins/Developer/Concert/ConcertApp/MultiUserServer/Source/MultiUserServer/Private/Widgets/Clients/Logging/Filter/ConcertFrontendLogFilter_Client.cpp
--- a/Engine/Plugins/Developer/Concert/ConcertApp/MultiUserServer/Source/MultiUserServer/Private/Widgets/Clients/Logging/Filter/ConcertFrontendLogFilter_Client.cpp
+++ b/Engine/Plugins/Developer/Concert/ConcertApp/MultiUserServer/Source/MultiUserServer/Private/Widgets/Clients/Logging/Filter/ConcertFrontendLogFilter_Client.cpp
@@ -6,12 +6,11 @@
 #include "Widgets/Clients/Logging/ConcertLogEntry.h"
 #include "Widgets/Clients/Util/EndpointToUserNameCache.h"
 
-namespace UE::MultiUserServer::Filters
+namespace UE::MultiUserServer
 {
 	bool FConcertLogFilter_Client::PassesFilter(const FConcertLogEntry& InItem) const
 	{
 		return AllowedClientMessagingNodeId == EndpointCache->TranslateEndpointIdToNodeId(InItem.Log.OriginEndpointId)
 			|| AllowedClientMessagingNodeId == EndpointCache->TranslateEndpointIdToNodeId(InItem.Log.DestinationEndpointId);
 	}
-}
-
+}
\ No newline at end of file
diff --git a/Engine/Plugins/Developer/Concert/ConcertApp/MultiUserServer/Source/MultiUserServer/Private/Widgets/Clients/Logging/Filter/ConcertFrontendLogFilter_Client.h b/Engine/Plugins/Developer/Concert/ConcertApp/MultiUserServer/Source/MultiUserServer/Private/Widgets/Clients/Logging/Filter/ConcertFrontendLogFilter_Client.h
--- a/Engine/Plugins/Developer/Concert/ConcertApp/MultiUserServer/Source/MultiUserServer/Private/Widgets/Clients/Logging/Filter/ConcertFrontendLogFilter_Client.h
+++ b/Engine/Plugins/Developer/Concert/ConcertApp/MultiUserServer/Source/MultiUserServer/Private/Widgets/Clients/Logging/Filter/ConcertFrontendLogFilter_Client.h
@@ -3,11 +3,12 @@
 #pragma once
 
 #include "CoreMinimal.h"
-#include "Widgets/Clients/Logging/Filter/ConcertFrontendLogFilter.h"
+#include "ConcertLogFilterTypes.h"
+#include "Widgets/Util/Filter/ConcertFilter.h"
 
 class FEndpointToUserNameCache;
 
-namespace UE::MultiUserServer::Filters
+namespace UE::MultiUserServer
 {
 	/** Only allows messages from the given clients */
 	class FConcertLogFilter_Client : public FConcertLogFilter
diff --git a/Engine/Plugins/Developer/Concert/ConcertApp/MultiUserServer/Source/MultiUserServer/Private/Widgets/Clients/Logging/Filter/ConcertFrontendLogFilter_MessageAction.h b/Engine/Plugins/Developer/Concert/ConcertApp/MultiUserServer/Source/MultiUserServer/Private/Widgets/Clients/Logging/Filter/ConcertFrontendLogFilter_MessageAction.h
--- a/Engine/Plugins/Developer/Concert/ConcertApp/MultiUserServer/Source/MultiUserServer/Private/Widgets/Clients/Logging/Filter/ConcertFrontendLogFilter_MessageAction.h
+++ b/Engine/Plugins/Developer/Concert/ConcertApp/MultiUserServer/Source/MultiUserServer/Private/Widgets/Clients/Logging/Filter/ConcertFrontendLogFilter_MessageAction.h
@@ -5,10 +5,11 @@
 #include "CoreMinimal.h"
 #include "ConcertFrontendLogFilter_BaseSetSelection.h"
 #include "Widgets/Clients/Logging/ConcertLogEntry.h"
-#include "Widgets/Clients/Logging/Filter/ConcertFrontendLogFilter.h"
 #include "Widgets/Clients/Logging/Util/MessageActionUtils.h"
+#include "Widgets/Util/Filter/ConcertFilter.h"
+#include "Widgets/Util/Filter/ConcertFrontendFilter.h"
 
-namespace UE::MultiUserServer::Filters
+namespace UE::MultiUserServer
 {
 	/** Allows only the selected messages types */
 	class FConcertLogFilter_MessageAction : public TConcertLogFilter_BaseSetSelection<FConcertLogFilter_MessageAction, FName>
diff --git a/Engine/Plugins/Developer/Concert/ConcertApp/MultiUserServer/Source/MultiUserServer/Private/Widgets/Clients/Logging/Filter/ConcertFrontendLogFilter_MessageType.h b/Engine/Plugins/Developer/Concert/ConcertApp/MultiUserServer/Source/MultiUserServer/Private/Widgets/Clients/Logging/Filter/ConcertFrontendLogFilter_MessageType.h
--- a/Engine/Plugins/Developer/Concert/ConcertApp/MultiUserServer/Source/MultiUserServer/Private/Widgets/Clients/Logging/Filter/ConcertFrontendLogFilter_MessageType.h
+++ b/Engine/Plugins/Developer/Concert/ConcertApp/MultiUserServer/Source/MultiUserServer/Private/Widgets/Clients/Logging/Filter/ConcertFrontendLogFilter_MessageType.h
@@ -6,8 +6,10 @@
 #include "ConcertFrontendLogFilter_BaseSetSelection.h"
 #include "Widgets/Clients/Logging/ConcertLogEntry.h"
 #include "Widgets/Clients/Logging/Util/MessageTypeUtils.h"
+#include "Widgets/Util/Filter/ConcertFilter.h"
+#include "Widgets/Util/Filter/ConcertFrontendFilter.h"
 
-namespace UE::MultiUserServer::Filters
+namespace UE::MultiUserServer
 {
 	/** Allows only the selected messages types */
 	class FConcertLogFilter_MessageType : public TConcertLogFilter_BaseSetSelection<FConcertLogFilter_MessageType, FName>
diff --git a/Engine/Plugins/Developer/Concert/ConcertApp/MultiUserServer/Source/MultiUserServer/Private/Widgets/Clients/Logging/Filter/ConcertFrontendLogFilter_Size.cpp b/Engine/Plugins/Developer/Concert/ConcertApp/MultiUserServer/Source/MultiUserServer/Private/Widgets/Clients/Logging/Filter/ConcertFrontendLogFilter_Size.cpp
--- a/Engine/Plugins/Developer/Concert/ConcertApp/MultiUserServer/Source/MultiUserServer/Private/Widgets/Clients/Logging/Filter/ConcertFrontendLogFilter_Size.cpp
+++ b/Engine/Plugins/Developer/Concert/ConcertApp/MultiUserServer/Source/MultiUserServer/Private/Widgets/Clients/Logging/Filter/ConcertFrontendLogFilter_Size.cpp
@@ -2,7 +2,6 @@
 
 #include "ConcertFrontendLogFilter_Size.h"
 
-#include "ConcertFrontendUtils.h"
 #include "ConcertTransportEvents.h"
 #include "SSimpleComboButton.h"
 
@@ -15,150 +14,153 @@
 
 #define LOCTEXT_NAMESPACE "UnrealMultiUserUI.FConcertLogFilter_Size"
 
-bool FConcertLogFilter_Size::PassesFilter(const FConcertLogEntry& InItem) const
+namespace UE::MultiUserServer
 {
-	check(InItem.Log.CustomPayloadUncompressedByteSize >= 0);
-	// Note: This only filters activities events - they all use custom events.
-	// The filter's default value is to show everything 0 <= x so it will show sync events as well.
-	const uint32 ComparisionSizeInBytes = FUnitConversion::Convert(SizeInBytes, DataUnit, EUnit::Bytes);
-	switch (FilterMode)
+	bool FConcertLogFilter_Size::PassesFilter(const FConcertLogEntry& InItem) const
 	{
-	case ESizeFilterMode::LessThanOrEqual:
-		return ComparisionSizeInBytes >= static_cast<uint32>(InItem.Log.CustomPayloadUncompressedByteSize);
-	case ESizeFilterMode::BiggerThanOrEqual:
-		return ComparisionSizeInBytes <= static_cast<uint32>(InItem.Log.CustomPayloadUncompressedByteSize);
-	default:
-		checkNoEntry();
-		return false;
+		check(InItem.Log.CustomPayloadUncompressedByteSize >= 0);
+		// Note: This only filters activities events - they all use custom events.
+		// The filter's default value is to show everything 0 <= x so it will show sync events as well.
+		const uint32 ComparisionSizeInBytes = FUnitConversion::Convert(SizeInBytes, DataUnit, EUnit::Bytes);
+		switch (FilterMode)
+		{
+		case ESizeFilterMode::LessThanOrEqual:
+			return ComparisionSizeInBytes >= static_cast<uint32>(InItem.Log.CustomPayloadUncompressedByteSize);
+		case ESizeFilterMode::BiggerThanOrEqual:
+			return ComparisionSizeInBytes <= static_cast<uint32>(InItem.Log.CustomPayloadUncompressedByteSize);
+		default:
+			checkNoEntry();
+			return false;
+		}
 	}
-}
 
-void FConcertLogFilter_Size::AdvanceFilterMode()
-{
-	switch (FilterMode)
+	void FConcertLogFilter_Size::AdvanceFilterMode()
 	{
-	case ESizeFilterMode::LessThanOrEqual:
-		FilterMode = ESizeFilterMode::BiggerThanOrEqual;
-		break;
-	case ESizeFilterMode::BiggerThanOrEqual:
-		FilterMode = ESizeFilterMode::LessThanOrEqual;
-		break;
-	default:
-		checkNoEntry();
-	}
+		switch (FilterMode)
+		{
+		case ESizeFilterMode::LessThanOrEqual:
+			FilterMode = ESizeFilterMode::BiggerThanOrEqual;
+			break;
+		case ESizeFilterMode::BiggerThanOrEqual:
+			FilterMode = ESizeFilterMode::LessThanOrEqual;
+			break;
+		default:
+			checkNoEntry();
+		}
 
-	OnChanged().Broadcast();
-}
+		OnChanged().Broadcast();
+	}
 
-void FConcertLogFilter_Size::SetSizeInBytes(uint32 NewSizeInBytes)
-{
-	if (NewSizeInBytes != SizeInBytes)
+	void FConcertLogFilter_Size::SetSizeInBytes(uint32 NewSizeInBytes)
 	{
-		SizeInBytes = NewSizeInBytes;
-		OnChanged().Broadcast();
+		if (NewSizeInBytes != SizeInBytes)
+		{
+			SizeInBytes = NewSizeInBytes;
+			OnChanged().Broadcast();
+		}
 	}
-}
 
-void FConcertLogFilter_Size::SetDataUnit(EUnit NewUnit)
-{
-	if (DataUnit != NewUnit && ensure(GetAllowedUnits().Contains(NewUnit)))
+	void FConcertLogFilter_Size::SetDataUnit(EUnit NewUnit)
 	{
-		DataUnit = NewUnit;
-		OnChanged().Broadcast();
+		if (DataUnit != NewUnit && ensure(GetAllowedUnits().Contains(NewUnit)))
+		{
+			DataUnit = NewUnit;
+			OnChanged().Broadcast();
+		}
 	}
-}
 
-FConcertFrontendLogFilter_Size::FConcertFrontendLogFilter_Size()
-{
-	ChildSlot = SNew(SHorizontalBox)
-		+SHorizontalBox::Slot()
-		.AutoWidth()
-		.HAlign(HAlign_Center)
-		[
-			SNew(SButton)
-			.OnClicked_Lambda([this]()
-			{
-				Implementation.AdvanceFilterMode();
-				return FReply::Handled();
-			})
-			.ToolTipText_Lambda([this]()
-			{
-				switch (Implementation.GetFilterMode())
-				{
-				case ESizeFilterMode::BiggerThanOrEqual: return FText::Format(LOCTEXT("BiggerThanOrEqual.ToolTipFmt", "Size >= {0}"), GetSizeAndUnitAsText());
-					case ESizeFilterMode::LessThanOrEqual: return FText::Format(LOCTEXT("LessThanOrEqual.ToolTipFmt", "Size <= {0}"), GetSizeAndUnitAsText());
-					default: return FText::GetEmpty();
-				}
-			})
+	FConcertFrontendLogFilter_Size::FConcertFrontendLogFilter_Size()
+	{
+		ChildSlot = SNew(SHorizontalBox)
+			+SHorizontalBox::Slot()
+			.AutoWidth()
+			.HAlign(HAlign_Center)
 			[
-				SNew(STextBlock)
-				.Text_Lambda([this]()
+				SNew(SButton)
+				.OnClicked_Lambda([this]()
+				{
+					Implementation.AdvanceFilterMode();
+					return FReply::Handled();
+				})
+				.ToolTipText_Lambda([this]()
 				{
 					switch (Implementation.GetFilterMode())
 					{
-						case ESizeFilterMode::BiggerThanOrEqual: return LOCTEXT("BiggerThanOrEqual.Text", ">=") ;
-						case ESizeFilterMode::LessThanOrEqual: return LOCTEXT("LessThanOrEqual.Text", "<=") ;
+					case ESizeFilterMode::BiggerThanOrEqual: return FText::Format(LOCTEXT("BiggerThanOrEqual.ToolTipFmt", "Size >= {0}"), GetSizeAndUnitAsText());
+						case ESizeFilterMode::LessThanOrEqual: return FText::Format(LOCTEXT("LessThanOrEqual.ToolTipFmt", "Size <= {0}"), GetSizeAndUnitAsText());
 						default: return FText::GetEmpty();
 					}
 				})
+				[
+					SNew(STextBlock)
+					.Text_Lambda([this]()
+					{
+						switch (Implementation.GetFilterMode())
+						{
+							case ESizeFilterMode::BiggerThanOrEqual: return LOCTEXT("BiggerThanOrEqual.Text", ">=") ;
+							case ESizeFilterMode::LessThanOrEqual: return LOCTEXT("LessThanOrEqual.Text", "<=") ;
+							default: return FText::GetEmpty();
+						}
+					})
+				]
 			]
-		]
-	
-		+SHorizontalBox::Slot()
-		.AutoWidth()
-		[
-			SNew(SNumericEntryBox<uint32>)
-			.AllowSpin(true)
-			.MinDesiredValueWidth(30)
-			.MaxSliderValue(1000)
-			.OnValueChanged_Lambda([this](uint32 NewValue){ Implementation.SetSizeInBytes(NewValue); })
-			.Value_Lambda([this](){ return Implementation.GetSizeInBytes(); })
-		]
-
-		+SHorizontalBox::Slot()
-		.AutoWidth()
-		[
-			SNew(SSimpleComboButton)
-			.HasDownArrow(true)
-			.Text_Lambda([this](){ return FText::FromString(FUnitConversion::GetUnitDisplayString(Implementation.GetDataUnit())); })
-			.MenuContent()
+		
+			+SHorizontalBox::Slot()
+			.AutoWidth()
 			[
-				MakeDataUnitMenu()
+				SNew(SNumericEntryBox<uint32>)
+				.AllowSpin(true)
+				.MinDesiredValueWidth(30)
+				.MaxSliderValue(1000)
+				.OnValueChanged_Lambda([this](uint32 NewValue){ Implementation.SetSizeInBytes(NewValue); })
+				.Value_Lambda([this](){ return Implementation.GetSizeInBytes(); })
 			]
-		];
-}
 
-TSharedRef<SWidget> FConcertFrontendLogFilter_Size::MakeDataUnitMenu()
-{
-	FMenuBuilder MenuBuilder(true, nullptr);
+			+SHorizontalBox::Slot()
+			.AutoWidth()
+			[
+				SNew(SSimpleComboButton)
+				.HasDownArrow(true)
+				.Text_Lambda([this](){ return FText::FromString(FUnitConversion::GetUnitDisplayString(Implementation.GetDataUnit())); })
+				.MenuContent()
+				[
+					MakeDataUnitMenu()
+				]
+			];
+	}
 
-	for (const EUnit Unit : Implementation.GetAllowedUnits())
+	TSharedRef<SWidget> FConcertFrontendLogFilter_Size::MakeDataUnitMenu()
 	{
-		MenuBuilder.AddMenuEntry(
-			FText::FromString(FUnitConversion::GetUnitDisplayString(Unit)),
-			FText::GetEmpty(),
-			FSlateIcon(),
-			FUIAction(
-				FExecuteAction::CreateLambda([this, Unit]()
-				{
-					Implementation.SetDataUnit(Unit);
-				}),
-				FCanExecuteAction::CreateLambda([] { return true; })),
-			NAME_None,
-			EUserInterfaceActionType::Button
-			);
+		FMenuBuilder MenuBuilder(true, nullptr);
+
+		for (const EUnit Unit : Implementation.GetAllowedUnits())
+		{
+			MenuBuilder.AddMenuEntry(
+				FText::FromString(FUnitConversion::GetUnitDisplayString(Unit)),
+				FText::GetEmpty(),
+				FSlateIcon(),
+				FUIAction(
+					FExecuteAction::CreateLambda([this, Unit]()
+					{
+						Implementation.SetDataUnit(Unit);
+					}),
+					FCanExecuteAction::CreateLambda([] { return true; })),
+				NAME_None,
+				EUserInterfaceActionType::Button
+				);
+		}
+		
+		return MenuBuilder.MakeWidget();
 	}
-	
-	return MenuBuilder.MakeWidget();
-}
 
-FText FConcertFrontendLogFilter_Size::GetSizeAndUnitAsText() const
-{
-	return FText::Format(
-		LOCTEXT("SizeAndUnitAsTextFmt", "{0} {1}"),
-		FUnitConversion::Convert(Implementation.GetSizeInBytes(), Implementation.GetDataUnit(), EUnit::Bytes),
-		FText::FromString(FUnitConversion::GetUnitDisplayString(Implementation.GetDataUnit()))
-		);
+	FText FConcertFrontendLogFilter_Size::GetSizeAndUnitAsText() const
+	{
+		return FText::Format(
+			LOCTEXT("SizeAndUnitAsTextFmt", "{0} {1}"),
+			FUnitConversion::Convert(Implementation.GetSizeInBytes(), Implementation.GetDataUnit(), EUnit::Bytes),
+			FText::FromString(FUnitConversion::GetUnitDisplayString(Implementation.GetDataUnit()))
+			);
+	}
 }
 
 #undef LOCTEXT_NAMESPACE
diff --git a/Engine/Plugins/Developer/Concert/ConcertApp/MultiUserServer/Source/MultiUserServer/Private/Widgets/Clients/Logging/Filter/ConcertFrontendLogFilter_Size.h b/Engine/Plugins/Developer/Concert/ConcertApp/MultiUserServer/Source/MultiUserServer/Private/Widgets/Clients/Logging/Filter/ConcertFrontendLogFilter_Size.h
--- a/Engine/Plugins/Developer/Concert/ConcertApp/MultiUserServer/Source/MultiUserServer/Private/Widgets/Clients/Logging/Filter/ConcertFrontendLogFilter_Size.h
+++ b/Engine/Plugins/Developer/Concert/ConcertApp/MultiUserServer/Source/MultiUserServer/Private/Widgets/Clients/Logging/Filter/ConcertFrontendLogFilter_Size.h
@@ -3,51 +3,58 @@
 #pragma once
 
 #include "CoreMinimal.h"
-#include "ConcertFrontendLogFilter.h"
+
+#include "ConcertLogFilterTypes.h"
+#include "Widgets/Util/Filter/ConcertFilter.h"
+#include "Widgets/Util/Filter/ConcertFrontendFilter.h"
+
 #include "Math/UnitConversion.h"
 
-enum class ESizeFilterMode
+namespace UE::MultiUserServer
 {
-	/** Allow logs bigger than or equal to the specified size */
-	BiggerThanOrEqual,
-	/** Allow logs smaller than or equal to the specified size */
-	LessThanOrEqual
-};
-
-/** Filters based on the log's size */
-class FConcertLogFilter_Size : public FConcertLogFilter
-{
-public:
+	enum class ESizeFilterMode
+	{
+		/** Allow logs bigger than or equal to the specified size */
+		BiggerThanOrEqual,
+		/** Allow logs smaller than or equal to the specified size */
+		LessThanOrEqual
+	};
+
+	/** Filters based on the log's size */
+	class FConcertLogFilter_Size : public FConcertLogFilter
+	{
+	public:
 	
-	//~ Begin FConcertLogFilter Interface
-	virtual bool PassesFilter(const FConcertLogEntry& InItem) const override;
-	//~ End FConcertLogFilter Interface
+		//~ Begin FConcertLogFilter Interface
+		virtual bool PassesFilter(const FConcertLogEntry& InItem) const override;
+		//~ End FConcertLogFilter Interface
 
-	void AdvanceFilterMode();
-	void SetSizeInBytes(uint32 NewSizeInBytes);
-	void SetDataUnit(EUnit NewUnit);
+		void AdvanceFilterMode();
+		void SetSizeInBytes(uint32 NewSizeInBytes);
+		void SetDataUnit(EUnit NewUnit);
 	
-	ESizeFilterMode GetFilterMode() const { return FilterMode; }
-	uint32 GetSizeInBytes() const { return SizeInBytes; }
-	EUnit GetDataUnit() const { return DataUnit; }
-	TSet<EUnit> GetAllowedUnits() const { return { EUnit::Bytes, EUnit::Kilobytes, EUnit::Megabytes }; }
+		ESizeFilterMode GetFilterMode() const { return FilterMode; }
+		uint32 GetSizeInBytes() const { return SizeInBytes; }
+		EUnit GetDataUnit() const { return DataUnit; }
+		TSet<EUnit> GetAllowedUnits() const { return { EUnit::Bytes, EUnit::Kilobytes, EUnit::Megabytes }; }
 	
-private:
+	private:
 
-	ESizeFilterMode FilterMode = ESizeFilterMode::BiggerThanOrEqual;
-	uint32 SizeInBytes = 0;
-	EUnit DataUnit = EUnit::Bytes;
-};
+		ESizeFilterMode FilterMode = ESizeFilterMode::BiggerThanOrEqual;
+		uint32 SizeInBytes = 0;
+		EUnit DataUnit = EUnit::Bytes;
+	};
 
-class FConcertFrontendLogFilter_Size : public TConcertFrontendLogFilterAggregate<FConcertLogFilter_Size>
-{
-	using Super = TConcertFrontendLogFilterAggregate<FConcertLogFilter_Size>;
-public:
+	class FConcertFrontendLogFilter_Size : public TConcertFrontendFilterAggregate<FConcertLogFilter_Size, const FConcertLogEntry&>
+	{
+		using Super = TConcertFrontendFilterAggregate<FConcertLogFilter_Size, const FConcertLogEntry&>;
+	public:
 	
-	FConcertFrontendLogFilter_Size();
+		FConcertFrontendLogFilter_Size();
 
-private:
+	private:
 
-	TSharedRef<SWidget> MakeDataUnitMenu();
-	FText GetSizeAndUnitAsText() const;
-};
+		TSharedRef<SWidget> MakeDataUnitMenu();
+		FText GetSizeAndUnitAsText() const;
+	};
+}
diff --git a/Engine/Plugins/Developer/Concert/ConcertApp/MultiUserServer/Source/MultiUserServer/Private/Widgets/Clients/Logging/Filter/ConcertFrontendLogFilter_TextSearch.cpp b/Engine/Plugins/Developer/Concert/ConcertApp/MultiUserServer/Source/MultiUserServer/Private/Widgets/Clients/Logging/Filter/ConcertFrontendLogFilter_TextSearch.cpp
deleted file mode 100644
diff --git a/Engine/Plugins/Developer/Concert/ConcertApp/MultiUserServer/Source/MultiUserServer/Private/Widgets/Clients/Logging/Filter/ConcertFrontendLogFilter_TextSearch.h b/Engine/Plugins/Developer/Concert/ConcertApp/MultiUserServer/Source/MultiUserServer/Private/Widgets/Clients/Logging/Filter/ConcertFrontendLogFilter_TextSearch.h
--- a/Engine/Plugins/Developer/Concert/ConcertApp/MultiUserServer/Source/MultiUserServer/Private/Widgets/Clients/Logging/Filter/ConcertFrontendLogFilter_TextSearch.h
+++ b/Engine/Plugins/Developer/Concert/ConcertApp/MultiUserServer/Source/MultiUserServer/Private/Widgets/Clients/Logging/Filter/ConcertFrontendLogFilter_TextSearch.h
@@ -3,47 +3,48 @@
 #pragma once
 
 #include "CoreMinimal.h"
-#include "ConcertFrontendLogFilter.h"
+#include "ConcertTransportEvents.h"
 #include "Misc/TextFilter.h"
-#include "Widgets/Input/SSearchBox.h"
+#include "Widgets/Clients/Logging/ConcertLogEntry.h"
+#include "Widgets/Clients/Logging/Util/ConcertLogTokenizer.h"
+#include "Widgets/Util/Filter/ConcertFrontendFilter_TextSearch.h"
 
 class FConcertLogTokenizer;
 
-/** Allows advanced search by text. Implements Adapter pattern to wrap TTextFilter. */
-class FConcertLogFilter_TextSearch : public FConcertLogFilter
+namespace UE::MultiUserServer
 {
-public:
+	class FConcertLogFilter_TextSearch : public TConcertFilter_TextSearch<const FConcertLogEntry&>
+	{
+	public:
 
-	FConcertLogFilter_TextSearch(TSharedRef<FConcertLogTokenizer> Tokenizer);
+		FConcertLogFilter_TextSearch(TSharedRef<FConcertLogTokenizer> Tokenizer)
+			: Tokenizer(MoveTemp(Tokenizer))
+		{}
 	
-	//~ Begin FConcertLogFilter Interface
-	virtual bool PassesFilter(const FConcertLogEntry& InItem) const override { return TextFilter.PassesFilter(InItem); }
-	//~ End FConcertLogFilter Interface
+		virtual void GenerateSearchTerms(const FConcertLogEntry& InItem, TArray<FString>& OutTerms) const override
+		{
+			for (TFieldIterator<const FProperty> PropertyIt(FConcertLog::StaticStruct()); PropertyIt; ++PropertyIt)
+			{
+				OutTerms.Add(Tokenizer->Tokenize(InItem.Log, **PropertyIt));
+			}
+		}
 	
-	void SetRawFilterText(const FText& InFilterText) { TextFilter.SetRawFilterText(InFilterText); }
-	
-private:
-
-	/** Does the actual string search */
-	TTextFilter<const FConcertLogEntry&> TextFilter;
-	/** Helps in converting FConcertLog members into search terms */
-	TSharedRef<FConcertLogTokenizer> Tokenizer;
+	private:
 
-	/** Parses InItem into a bunch of strings that can be searched */
-	void GenerateSearchTerms(const FConcertLogEntry& InItem, TArray<FString>& OutTerms) const;
-};
-
-/** Creates a search bar */
-class FConcertFrontendLogFilter_TextSearch : public TConcertFrontendLogFilterAggregate<FConcertLogFilter_TextSearch, SSearchBox>
-{
-	using Super = TConcertFrontendLogFilterAggregate<FConcertLogFilter_TextSearch, SSearchBox>;
-public:
+		/** Helps in converting FConcertLog members into search terms */
+		TSharedRef<FConcertLogTokenizer> Tokenizer;
+	};
 
-	FConcertFrontendLogFilter_TextSearch(TSharedRef<FConcertLogTokenizer> Tokenizer);
+	/** Creates a search bar */
+	class FConcertFrontendLogFilter_TextSearch : public TConcertFrontendFilter_TextSearch<FConcertLogFilter_TextSearch, const FConcertLogEntry&>
+	{
+		using Super = TConcertFrontendFilter_TextSearch<FConcertLogFilter_TextSearch, const FConcertLogEntry&>;
+	public:
 
-	FText GetSearchText() const { return SearchText; }
+		FConcertFrontendLogFilter_TextSearch(TSharedRef<FConcertLogTokenizer> Tokenizer)
+			: Super(MoveTemp(Tokenizer))
+		{}
+	};
+}
 
-private:
 
-	FText SearchText;
-};
diff --git a/Engine/Plugins/Developer/Concert/ConcertApp/MultiUserServer/Source/MultiUserServer/Private/Widgets/Clients/Logging/Filter/ConcertFrontendLogFilter_Time.cpp b/Engine/Plugins/Developer/Concert/ConcertApp/MultiUserServer/Source/MultiUserServer/Private/Widgets/Clients/Logging/Filter/ConcertFrontendLogFilter_Time.cpp
--- a/Engine/Plugins/Developer/Concert/ConcertApp/MultiUserServer/Source/MultiUserServer/Private/Widgets/Clients/Logging/Filter/ConcertFrontendLogFilter_Time.cpp
+++ b/Engine/Plugins/Developer/Concert/ConcertApp/MultiUserServer/Source/MultiUserServer/Private/Widgets/Clients/Logging/Filter/ConcertFrontendLogFilter_Time.cpp
@@ -2,7 +2,6 @@
 
 #include "ConcertFrontendLogFilter_Time.h"
 
-#include "ConcertFrontendUtils.h"
 #include "ConcertTransportEvents.h"
 #include "Widgets/Clients/Logging/ConcertLogEntry.h"
 
@@ -12,130 +11,133 @@
 #include "Widgets/SBoxPanel.h"
 #include "Widgets/Text/STextBlock.h"
 
-#define LOCTEXT_NAMESPACE "UnrealMultiUserUI"
+#define LOCTEXT_NAMESPACE "UnrealMultiUserUI.FConcertFrontendLogFilter_Time"
 
-FConcertLogFilter_Time::FConcertLogFilter_Time(ETimeFilter FilterMode)
-	: FilterMode(FilterMode)
-	// Make the filter allow everything by default
-	, Time(MakeResetTime())
-{}
-
-void FConcertLogFilter_Time::ResetToInfiniteTime()
+namespace UE::MultiUserServer
 {
-	SetTime(MakeResetTime());
-}
+	FConcertLogFilter_Time::FConcertLogFilter_Time(ETimeFilter FilterMode)
+		: FilterMode(FilterMode)
+		// Make the filter allow everything by default
+		, Time(MakeResetTime())
+	{}
 
-FDateTime FConcertLogFilter_Time::MakeResetTime() const
-{
-	return FilterMode == ETimeFilter::AllowAfter ? FDateTime() : FDateTime::MaxValue();
-}
+	void FConcertLogFilter_Time::ResetToInfiniteTime()
+	{
+		SetTime(MakeResetTime());
+	}
 
-bool FConcertLogFilter_Time::PassesFilter(const FConcertLogEntry& InItem) const
-{
-	switch (FilterMode)
+	FDateTime FConcertLogFilter_Time::MakeResetTime() const
 	{
-	case ETimeFilter::AllowAfter: return InItem.Log.Timestamp >= Time;
-	case ETimeFilter::AllowBefore: return InItem.Log.Timestamp <= Time;
-	default:
-		checkNoEntry();
-		return true;
+		return FilterMode == ETimeFilter::AllowAfter ? FDateTime() : FDateTime::MaxValue();
 	}
-}
 
-void FConcertLogFilter_Time::SetFilterMode(ETimeFilter InFilterMode)
-{
-	if (FilterMode != InFilterMode)
+	bool FConcertLogFilter_Time::PassesFilter(const FConcertLogEntry& InItem) const
 	{
-		FilterMode = InFilterMode;
-		OnChanged().Broadcast();
+		switch (FilterMode)
+		{
+		case ETimeFilter::AllowAfter: return InItem.Log.Timestamp >= Time;
+		case ETimeFilter::AllowBefore: return InItem.Log.Timestamp <= Time;
+		default:
+			checkNoEntry();
+			return true;
+		}
 	}
-}
 
-void FConcertLogFilter_Time::SetTime(const FDateTime& InTime)
-{
-	if (Time != InTime)
+	void FConcertLogFilter_Time::SetFilterMode(ETimeFilter InFilterMode)
 	{
-		Time = InTime;
-		OnChanged().Broadcast();
+		if (FilterMode != InFilterMode)
+		{
+			FilterMode = InFilterMode;
+			OnChanged().Broadcast();
+		}
 	}
-}
 
-FConcertFrontendLogFilter_Time::FConcertFrontendLogFilter_Time(ETimeFilter TimeFilter)
-	: Super(TimeFilter)
-{
-	ChildSlot = SNew(SHorizontalBox)
-		.ToolTipText(LOCTEXT("TimeFilter.ToolTipText", "Filter logs by local time"))
-
-		+SHorizontalBox::Slot()
-		.AutoWidth()
-		.VAlign(VAlign_Center)
-		[
-			SNew(STextBlock)
-			.Text_Lambda([this]()
-			{
-				return Implementation.GetFilterMode() == ETimeFilter::AllowAfter
-					? LOCTEXT("TimeFilter.AllowAfter", "After")
-					: LOCTEXT("TimeFilter.AllowBefore", "Before");
-			})
-		]
-
-		+SHorizontalBox::Slot()
-		.AutoWidth()
-		.Padding(2, 0, 0, 0)
-		.VAlign(VAlign_Center)
-		[
-			SNew(SComboButton)
-			.OnGetMenuContent_Raw(this, &FConcertFrontendLogFilter_Time::CreateDatePicker)
-			.ButtonContent()
+	void FConcertLogFilter_Time::SetTime(const FDateTime& InTime)
+	{
+		if (Time != InTime)
+		{
+			Time = InTime;
+			OnChanged().Broadcast();
+		}
+	}
+
+	FConcertFrontendLogFilter_Time::FConcertFrontendLogFilter_Time(ETimeFilter TimeFilter)
+		: Super(TimeFilter)
+	{
+		ChildSlot = SNew(SHorizontalBox)
+			.ToolTipText(LOCTEXT("TimeFilter.ToolTipText", "Filter logs by local time"))
+
+			+SHorizontalBox::Slot()
+			.AutoWidth()
+			.VAlign(VAlign_Center)
 			[
 				SNew(STextBlock)
-				.Text_Lambda([this](){ return FText::AsDateTime(Implementation.GetTime()); })
+				.Text_Lambda([this]()
+				{
+					return Implementation.GetFilterMode() == ETimeFilter::AllowAfter
+						? LOCTEXT("TimeFilter.AllowAfter", "After")
+						: LOCTEXT("TimeFilter.AllowBefore", "Before");
+				})
 			]
-		];
-}
 
-TSharedRef<SWidget> FConcertFrontendLogFilter_Time::CreateDatePicker()
-{
-	FMenuBuilder MenuBuilder(true, nullptr);
-
-	MenuBuilder.AddMenuEntry(
-		LOCTEXT("TimeFilter.Clear.", "Clear"),
-		LOCTEXT("TimeFilter.Clear.Tooltip", "Sets the time so that this filter has no effect"),
-		FSlateIcon(),
-		FUIAction(
-			FExecuteAction::CreateLambda([this](){ Implementation.ResetToInfiniteTime(); }),
-			FCanExecuteAction::CreateLambda([] { return true; }),
-			FIsActionChecked()),
-		NAME_None,
-		EUserInterfaceActionType::Button
+			+SHorizontalBox::Slot()
+			.AutoWidth()
+			.Padding(2, 0, 0, 0)
+			.VAlign(VAlign_Center)
+			[
+				SNew(SComboButton)
+				.OnGetMenuContent_Raw(this, &FConcertFrontendLogFilter_Time::CreateDatePicker)
+				.ButtonContent()
+				[
+					SNew(STextBlock)
+					.Text_Lambda([this](){ return FText::AsDateTime(Implementation.GetTime()); })
+				]
+			];
+	}
+
+	TSharedRef<SWidget> FConcertFrontendLogFilter_Time::CreateDatePicker()
+	{
+		FMenuBuilder MenuBuilder(true, nullptr);
+
+		MenuBuilder.AddMenuEntry(
+			LOCTEXT("TimeFilter.Clear.", "Clear"),
+			LOCTEXT("TimeFilter.Clear.Tooltip", "Sets the time so that this filter has no effect"),
+			FSlateIcon(),
+			FUIAction(
+				FExecuteAction::CreateLambda([this](){ Implementation.ResetToInfiniteTime(); }),
+				FCanExecuteAction::CreateLambda([] { return true; }),
+				FIsActionChecked()),
+			NAME_None,
+			EUserInterfaceActionType::Button
+			);
+		MenuBuilder.AddMenuEntry(
+			LOCTEXT("TimeFilter.Now.", "Now"),
+			LOCTEXT("TimeFilter.Now.Tooltip", "Sets the time to now in local time"),
+			FSlateIcon(),
+			FUIAction(
+				FExecuteAction::CreateLambda([this](){ Implementation.SetTime(FDateTime::UtcNow()); }),
+				FCanExecuteAction::CreateLambda([] { return true; }),
+				FIsActionChecked()),
+			NAME_None,
+			EUserInterfaceActionType::Button
 		);
-	MenuBuilder.AddMenuEntry(
-		LOCTEXT("TimeFilter.Now.", "Now"),
-		LOCTEXT("TimeFilter.Now.Tooltip", "Sets the time to now in local time"),
-		FSlateIcon(),
-		FUIAction(
-			FExecuteAction::CreateLambda([this](){ Implementation.SetTime(FDateTime::UtcNow()); }),
-			FCanExecuteAction::CreateLambda([] { return true; }),
-			FIsActionChecked()),
-		NAME_None,
-		EUserInterfaceActionType::Button
-	);
-	MenuBuilder.AddSeparator();
-
-	MenuBuilder.AddWidget(
-		SNew(SEditableTextBox)
-		.Text(FText::FromString(Implementation.GetTime().ToString()))
-		.OnTextCommitted_Lambda([this](const FText& Text, ETextCommit::Type)
-		{
-			FDateTime Time;
-			if (FDateTime::Parse(Text.ToString(), Time))
+		MenuBuilder.AddSeparator();
+
+		MenuBuilder.AddWidget(
+			SNew(SEditableTextBox)
+			.Text(FText::FromString(Implementation.GetTime().ToString()))
+			.OnTextCommitted_Lambda([this](const FText& Text, ETextCommit::Type)
 			{
-				Implementation.SetTime(Time);
-			}
-		}),
-		LOCTEXT("TimeFilter.Custom.", "Custom time")
-		);
-	return MenuBuilder.MakeWidget();
+				FDateTime Time;
+				if (FDateTime::Parse(Text.ToString(), Time))
+				{
+					Implementation.SetTime(Time);
+				}
+			}),
+			LOCTEXT("TimeFilter.Custom.", "Custom time")
+			);
+		return MenuBuilder.MakeWidget();
+	}
 }
 
 #undef LOCTEXT_NAMESPACE
diff --git a/Engine/Plugins/Developer/Concert/ConcertApp/MultiUserServer/Source/MultiUserServer/Private/Widgets/Clients/Logging/Filter/ConcertFrontendLogFilter_Time.h b/Engine/Plugins/Developer/Concert/ConcertApp/MultiUserServer/Source/MultiUserServer/Private/Widgets/Clients/Logging/Filter/ConcertFrontendLogFilter_Time.h
--- a/Engine/Plugins/Developer/Concert/ConcertApp/MultiUserServer/Source/MultiUserServer/Private/Widgets/Clients/Logging/Filter/ConcertFrontendLogFilter_Time.h
+++ b/Engine/Plugins/Developer/Concert/ConcertApp/MultiUserServer/Source/MultiUserServer/Private/Widgets/Clients/Logging/Filter/ConcertFrontendLogFilter_Time.h
@@ -3,51 +3,58 @@
 #pragma once
 
 #include "CoreMinimal.h"
-#include "ConcertFrontendLogFilter.h"
+#include "ConcertLogFilterTypes.h"
 
-enum class ETimeFilter
-{
-	/** Logs after the indicated time are allowed */
-	AllowAfter,
-	/** Logs before the indicated time are allowed */
-	AllowBefore
-};
-
-/** Filters based on whether a log happened before or after a certain time */
-class FConcertLogFilter_Time : public FConcertLogFilter
+#include "Widgets/Util/Filter/ConcertFilter.h"
+#include "Widgets/Util/Filter/ConcertFrontendFilter.h"
+
+namespace UE::MultiUserServer
 {
-public:
+	enum class ETimeFilter
+	{
+		/** Logs after the indicated time are allowed */
+		AllowAfter,
+		/** Logs before the indicated time are allowed */
+		AllowBefore
+	};
+
+	/** Filters based on whether a log happened before or after a certain time */
+	class FConcertLogFilter_Time : public FConcertLogFilter
+	{
+	public:
 
-	FConcertLogFilter_Time(ETimeFilter FilterMode);
+		FConcertLogFilter_Time(ETimeFilter FilterMode);
 
-	void ResetToInfiniteTime();
+		void ResetToInfiniteTime();
 
-	//~ Begin FConcertLogFilter Interface
-	virtual bool PassesFilter(const FConcertLogEntry& InItem) const override;
-	//~ End FConcertLogFilter Interfac
+		//~ Begin FConcertLogFilter Interface
+		virtual bool PassesFilter(const FConcertLogEntry& InItem) const override;
+		//~ End FConcertLogFilter Interfac
 
-	ETimeFilter GetFilterMode() const { return FilterMode; }
-	FDateTime GetTime() const { return Time; }
+		ETimeFilter GetFilterMode() const { return FilterMode; }
+		FDateTime GetTime() const { return Time; }
 
-	void SetFilterMode(ETimeFilter InFilterMode);
-	void SetTime(const FDateTime& InTime);
+		void SetFilterMode(ETimeFilter InFilterMode);
+		void SetTime(const FDateTime& InTime);
 	
-private:
+	private:
 
-	ETimeFilter FilterMode;
-	FDateTime Time;
+		ETimeFilter FilterMode;
+		FDateTime Time;
 
-	FDateTime MakeResetTime() const;
-};
+		FDateTime MakeResetTime() const;
+	};
 
-class FConcertFrontendLogFilter_Time : public TConcertFrontendLogFilterAggregate<FConcertLogFilter_Time>
-{
-	using Super = TConcertFrontendLogFilterAggregate<FConcertLogFilter_Time>;
-public:
+	class FConcertFrontendLogFilter_Time : public TConcertFrontendFilterAggregate<FConcertLogFilter_Time, const FConcertLogEntry&>
+	{
+		using Super = TConcertFrontendFilterAggregate<FConcertLogFilter_Time, const FConcertLogEntry&>;
+	public:
 	
-	FConcertFrontendLogFilter_Time(ETimeFilter TimeFilter);
+		FConcertFrontendLogFilter_Time(ETimeFilter TimeFilter);
+
+	private:
 
-private:
+		TSharedRef<SWidget> CreateDatePicker();
+	};
+}
 
-	TSharedRef<SWidget> CreateDatePicker();
-};
diff --git a/Engine/Plugins/Developer/Concert/ConcertApp/MultiUserServer/Source/MultiUserServer/Private/Widgets/Clients/Logging/Filter/ConcertLogFilter.h b/Engine/Plugins/Developer/Concert/ConcertApp/MultiUserServer/Source/MultiUserServer/Private/Widgets/Clients/Logging/Filter/ConcertLogFilter.h
deleted file mode 100644
diff --git a/Engine/Plugins/Developer/Concert/ConcertApp/MultiUserServer/Source/MultiUserServer/Private/Widgets/Clients/Logging/Filter/ConcertLogFilterTypes.h b/Engine/Plugins/Developer/Concert/ConcertApp/MultiUserServer/Source/MultiUserServer/Private/Widgets/Clients/Logging/Filter/ConcertLogFilterTypes.h
new file mode 100644
--- /dev/null
+++ b/Engine/Plugins/Developer/Concert/ConcertApp/MultiUserServer/Source/MultiUserServer/Private/Widgets/Clients/Logging/Filter/ConcertLogFilterTypes.h
@@ -0,0 +1,18 @@
+// Copyright Epic Games, Inc. All Rights Reserved.
+
+#pragma once
+
+#include "CoreMinimal.h"
+
+struct FConcertLogEntry;
+
+namespace UE::MultiUserServer
+{
+	template<typename TFilterType> class TConcertFilter;
+	template<typename TFilterType> class TConcertFrontendFilter;
+
+	using FConcertLogFilter = TConcertFilter<const FConcertLogEntry&>;
+	using FConcertFrontendLogFilter = TConcertFrontendFilter<const FConcertLogEntry&>;
+}
+
+
diff --git a/Engine/Plugins/Developer/Concert/ConcertApp/MultiUserServer/Source/MultiUserServer/Private/Widgets/Clients/Logging/Filter/ConcertLogFilter_FrontendRoot.cpp b/Engine/Plugins/Developer/Concert/ConcertApp/MultiUserServer/Source/MultiUserServer/Private/Widgets/Clients/Logging/Filter/ConcertLogFilter_FrontendRoot.cpp
--- a/Engine/Plugins/Developer/Concert/ConcertApp/MultiUserServer/Source/MultiUserServer/Private/Widgets/Clients/Logging/Filter/ConcertLogFilter_FrontendRoot.cpp
+++ b/Engine/Plugins/Developer/Concert/ConcertApp/MultiUserServer/Source/MultiUserServer/Private/Widgets/Clients/Logging/Filter/ConcertLogFilter_FrontendRoot.cpp
@@ -2,7 +2,6 @@
 
 #include "ConcertLogFilter_FrontendRoot.h"
 
-#include "ConcertFrontendLogFilter.h"
 #include "ConcertFrontendLogFilter_Ack.h"
 #include "ConcertFrontendLogFilter_Client.h"
 #include "ConcertFrontendLogFilter_MessageAction.h"
@@ -10,80 +9,8 @@
 #include "ConcertFrontendLogFilter_TextSearch.h"
 #include "ConcertFrontendLogFilter_Time.h"
 #include "ConcertFrontendLogFilter_Size.h"
-#include "Algo/AllOf.h"
+#include "ConcertLogFilterTypes.h"
 
-#include "Algo/AnyOf.h"
-#include "Widgets/SBoxPanel.h"
-
-FConcertLogFilter_FrontendRoot::FConcertLogFilter_FrontendRoot(TSharedRef<FConcertLogTokenizer> Tokenizer, TArray<TSharedRef<FConcertFrontendLogFilter>> InCustomFilters, const TArray<TSharedRef<FConcertLogFilter>>& NonVisualFilters)
-	: TextSearchFilter(MakeShared<FConcertFrontendLogFilter_TextSearch>(MoveTemp(Tokenizer)))
-	, FrontendFilters(MoveTemp(InCustomFilters))
-	, AllFilters(FrontendFilters)
-{
-	AllFilters.Reserve(FrontendFilters.Num() + 1 + NonVisualFilters.Num());
-	
-	AllFilters.Add(TextSearchFilter);
-	for (const TSharedRef<FConcertLogFilter>& NonVisualFilter : NonVisualFilters)
-	{
-		AllFilters.Add(NonVisualFilter);
-	}
-	
-	for (const TSharedRef<FConcertLogFilter>& Filter : AllFilters)
-	{
-		Filter->OnChanged().AddLambda([this]()
-		{
-			OnChanged().Broadcast();
-		});
-	}
-}
-
-TSharedRef<SWidget> FConcertLogFilter_FrontendRoot::BuildFilterWidgets() const
-{
-	return SNew(SVerticalBox)
-
-		// Search bar
-		+SVerticalBox::Slot()
-		.AutoHeight()
-		[
-			TextSearchFilter->GetFilterWidget()
-		]
-
-		+SVerticalBox::Slot()
-		.AutoHeight()
-		.Padding(0, 2)
-		[
-			BuildCustomFilterListWidget()
-		];
-}
-
-bool FConcertLogFilter_FrontendRoot::PassesFilter(const FConcertLogEntry& InItem) const
-{
-	return Algo::AllOf(
-		AllFilters,
-		[&InItem](const TSharedRef<FConcertLogFilter>& AndFilter){ return AndFilter->PassesFilter(InItem); }
-		);
-}
-
-TSharedRef<SWidget> FConcertLogFilter_FrontendRoot::BuildCustomFilterListWidget() const
-{
-	TSharedRef<SHorizontalBox> Box = SNew(SHorizontalBox);
-
-	bool bIsFirst = true;
-	for (const TSharedRef<FConcertFrontendLogFilter>& Filter : FrontendFilters)
-	{
-		const FMargin Margin = bIsFirst ? FMargin() : FMargin(8, 0, 0, 0);
-		bIsFirst = false;
-		Box->AddSlot()
-			.AutoWidth()
-			.Padding(Margin)
-			.VAlign(VAlign_Center)
-			[
-				Filter->GetFilterWidget()
-			];
-	}
-
-	return Box;
-}
 
 namespace UE::MultiUserServer
 {
@@ -92,12 +19,12 @@ namespace UE::MultiUserServer
 		static TArray<TSharedRef<FConcertFrontendLogFilter>> CreateCommonFilters()
 		{
 			return {
-				MakeShared<Filters::FConcertFrontendLogFilter_MessageAction>(),
-				MakeShared<Filters::FConcertFrontendLogFilter_MessageType>(),
+				MakeShared<FConcertFrontendLogFilter_MessageAction>(),
+				MakeShared<FConcertFrontendLogFilter_MessageType>(),
 				MakeShared<FConcertFrontendLogFilter_Time>(ETimeFilter::AllowAfter),
 				MakeShared<FConcertFrontendLogFilter_Time>(ETimeFilter::AllowBefore),
 				MakeShared<FConcertFrontendLogFilter_Size>(),
-				MakeShared<Filters::FConcertFrontendLogFilter_Ack>()
+				MakeShared<FConcertFrontendLogFilter_Ack>()
 			};
 		}
 	}
@@ -114,7 +41,7 @@ namespace UE::MultiUserServer
 	{
 		const TArray<TSharedRef<FConcertFrontendLogFilter>> CommonFilters = Private::CreateCommonFilters();
 		const TArray<TSharedRef<FConcertLogFilter>> NonVisuals = {
-			MakeShared<Filters::FConcertLogFilter_Client>(ClientMessageNodeId, EndpointCache)
+			MakeShared<FConcertLogFilter_Client>(ClientMessageNodeId, EndpointCache)
 		};
 		return MakeShared<FConcertLogFilter_FrontendRoot>(
 			MoveTemp(Tokenizer),
diff --git a/Engine/Plugins/Developer/Concert/ConcertApp/MultiUserServer/Source/MultiUserServer/Private/Widgets/Clients/Logging/Filter/ConcertLogFilter_FrontendRoot.h b/Engine/Plugins/Developer/Concert/ConcertApp/MultiUserServer/Source/MultiUserServer/Private/Widgets/Clients/Logging/Filter/ConcertLogFilter_FrontendRoot.h
--- a/Engine/Plugins/Developer/Concert/ConcertApp/MultiUserServer/Source/MultiUserServer/Private/Widgets/Clients/Logging/Filter/ConcertLogFilter_FrontendRoot.h
+++ b/Engine/Plugins/Developer/Concert/ConcertApp/MultiUserServer/Source/MultiUserServer/Private/Widgets/Clients/Logging/Filter/ConcertLogFilter_FrontendRoot.h
@@ -3,49 +3,30 @@
 #pragma once
 
 #include "CoreMinimal.h"
-#include "ConcertLogFilter.h"
+#include "ConcertFrontendLogFilter_TextSearch.h"
+#include "Widgets/Util/Filter/ConcertRootFrontendFilter.h"
 
 class FConcertLogTokenizer;
-class FConcertFrontendLogFilter;
-class FConcertFrontendLogFilter_TextSearch;
 class FEndpointToUserNameCache;
-class SWidget;
-
-/** A filter that contains multiple UI filters */
-class FConcertLogFilter_FrontendRoot : public FConcertLogFilter, public TSharedFromThis<FConcertLogFilter_FrontendRoot>
-{
-public:
-
-	FConcertLogFilter_FrontendRoot(TSharedRef<FConcertLogTokenizer> Tokenizer,
-		TArray<TSharedRef<FConcertFrontendLogFilter>> InCustomFilters,
-		const TArray<TSharedRef<FConcertLogFilter>>& NonVisualFilters = {});
-
-	/** Builds the widget view for all contained filters */
-	TSharedRef<SWidget> BuildFilterWidgets() const;
-	
-	//~ Begin IFilter Interface
-	virtual bool PassesFilter(const FConcertLogEntry& InItem) const override;
-	//~ End IFilter Interface
-
-	FORCEINLINE const TSharedRef<FConcertFrontendLogFilter_TextSearch>& GetTextSearchFilter() const { return TextSearchFilter; }
-
-private:
-
-	/** The text search filter. Also in FrontendFilters. Separate variable to build search bar in new line. */
-	TSharedRef<FConcertFrontendLogFilter_TextSearch> TextSearchFilter;
-
-	/** AllFrontendFilters without special filters we have as propteries above, such as TextSearchFilter. */
-	TArray<TSharedRef<FConcertFrontendLogFilter>> FrontendFilters;
-	
-	/** Filters that are combined using logical AND. */
-	TArray<TSharedRef<FConcertLogFilter>> AllFilters;
-
-	/** Builds the widgets that go under the text */
-	TSharedRef<SWidget> BuildCustomFilterListWidget() const;
-};
+struct FConcertLogEntry;
 
 namespace UE::MultiUserServer
 {
+	/** A filter that contains multiple UI filters */
+	class FConcertLogFilter_FrontendRoot : public TConcertFrontendRootFilter<const FConcertLogEntry&, FConcertFrontendLogFilter_TextSearch>
+	{
+		using Super = TConcertFrontendRootFilter<const FConcertLogEntry&, FConcertFrontendLogFilter_TextSearch>;
+	public:
+
+		FConcertLogFilter_FrontendRoot(
+			TSharedRef<FConcertLogTokenizer> Tokenizer,
+			TArray<TSharedRef<Super::FConcertFrontendFilter>> FrontendFilters,
+			const TArray<TSharedRef<Super::FConcertFilter>>& NonVisualFilters = {}
+			)
+			: Super(MakeShared<FConcertFrontendLogFilter_TextSearch>(MoveTemp(Tokenizer)), MoveTemp(FrontendFilters), NonVisualFilters)
+		{}
+	};
+	
 	/** Creates a filter for the global filter log window */
 	TSharedRef<FConcertLogFilter_FrontendRoot> MakeGlobalLogFilter(TSharedRef<FConcertLogTokenizer> Tokenizer);
 
diff --git a/Engine/Plugins/Developer/Concert/ConcertApp/MultiUserServer/Source/MultiUserServer/Private/Widgets/Clients/Logging/Filter/FilteredConcertLogList.cpp b/Engine/Plugins/Developer/Concert/ConcertApp/MultiUserServer/Source/MultiUserServer/Private/Widgets/Clients/Logging/Filter/FilteredConcertLogList.cpp
--- a/Engine/Plugins/Developer/Concert/ConcertApp/MultiUserServer/Source/MultiUserServer/Private/Widgets/Clients/Logging/Filter/FilteredConcertLogList.cpp
+++ b/Engine/Plugins/Developer/Concert/ConcertApp/MultiUserServer/Source/MultiUserServer/Private/Widgets/Clients/Logging/Filter/FilteredConcertLogList.cpp
@@ -2,220 +2,223 @@
 
 #include "FilteredConcertLogList.h"
 
-#include "ConcertLogFilter.h"
+#include "Widgets/Util/Filter/ConcertFilter.h"
 
-FFilteredConcertLogList::FFilteredConcertLogList(TSharedRef<IConcertLogSource> InLogSource, TSharedPtr<FConcertLogFilter> InOptionalFilter)
-	: LogSource(MoveTemp(InLogSource))
-	, Filter(MoveTemp(InOptionalFilter))
+namespace UE::MultiUserServer
 {
-	LogSource->OnLowestLogEntryChanged().AddRaw(this, &FFilteredConcertLogList::OnLowestLogEntryChanged);
-	LogSource->OnLogEntryAdded().AddRaw(this, &FFilteredConcertLogList::OnNewLogEntryAdded);
-
-	if (Filter)
+	FFilteredConcertLogList::FFilteredConcertLogList(TSharedRef<IConcertLogSource> InLogSource, TSharedPtr<FConcertLogFilter> InOptionalFilter)
+		: LogSource(MoveTemp(InLogSource))
+		, Filter(MoveTemp(InOptionalFilter))
 	{
-		Filter->OnChanged().AddRaw(this, &FFilteredConcertLogList::RebuildFilteredResult);
-	}
-
-	RebuildFilteredResult();
-}
+		LogSource->OnLowestLogEntryChanged().AddRaw(this, &FFilteredConcertLogList::OnLowestLogEntryChanged);
+		LogSource->OnLogEntryAdded().AddRaw(this, &FFilteredConcertLogList::OnNewLogEntryAdded);
 
-FFilteredConcertLogList::~FFilteredConcertLogList()
-{
-	// This is not really needed because the LogSource is SUPPOSED to be solely owned by FFilteredConcertLogList.
-	// We do it for safety regardless...
-	LogSource->OnLowestLogEntryChanged().RemoveAll(this);
-	LogSource->OnLogEntryAdded().RemoveAll(this);
+		if (Filter)
+		{
+			Filter->OnChanged().AddRaw(this, &FFilteredConcertLogList::RebuildFilteredResult);
+		}
 
-	// ... Filter is not uniquely owned but removing effectively also does not really matter.
-	if (Filter)
-	{
-		Filter->OnChanged().RemoveAll(this);
+		RebuildFilteredResult();
 	}
-}
 
-void FFilteredConcertLogList::OnLowestLogEntryChanged(FConcertLogID NewLowestValidID)
-{
-	for (size_t i = 0; i < FilteredLogs.Num() && FilteredLogs[i]->LogId < NewLowestValidID; ++i)
+	FFilteredConcertLogList::~FFilteredConcertLogList()
 	{
-		// Worst case this is O(n^2) - usually a thought O(cn) with small c)...
-		// Not much we without rewriting SListView to use something other than TArray
-		RemoveMessageIdToLogEntry(FilteredLogs[i]);
-		FilteredLogs.RemoveAt(i);
-		--i;
-	}
-}
+		// This is not really needed because the LogSource is SUPPOSED to be solely owned by FFilteredConcertLogList.
+		// We do it for safety regardless...
+		LogSource->OnLowestLogEntryChanged().RemoveAll(this);
+		LogSource->OnLogEntryAdded().RemoveAll(this);
 
-void FFilteredConcertLogList::OnNewLogEntryAdded(const TSharedRef<FConcertLogEntry>& NewLogEntry)
-{
-	if (!Filter || Filter->PassesFilter(*NewLogEntry))
-	{
-		MessageIdToLogs.FindOrAdd(NewLogEntry->Log.MessageId).AddUnique(NewLogEntry);
-		FilteredLogs.Emplace(NewLogEntry);
-		OnLogListChanged().Broadcast(GetFilteredLogs());
+		// ... Filter is not uniquely owned but removing effectively also does not really matter.
+		if (Filter)
+		{
+			Filter->OnChanged().RemoveAll(this);
+		}
 	}
-}
 
-void FFilteredConcertLogList::RebuildFilteredResult()
-{
-	MessageIdToLogs.Reset();
-	FilteredLogs.Reset();
-	LogSource->ForEachLog([this](const TSharedPtr<FConcertLogEntry>& LogEntry)
+	void FFilteredConcertLogList::OnLowestLogEntryChanged(FConcertLogID NewLowestValidID)
 	{
-		if (!Filter || Filter->PassesFilter(*LogEntry))
+		for (size_t i = 0; i < FilteredLogs.Num() && FilteredLogs[i]->LogId < NewLowestValidID; ++i)
 		{
-			MessageIdToLogs.FindOrAdd(LogEntry->Log.MessageId).AddUnique(LogEntry);
-			FilteredLogs.Emplace(LogEntry);
+			// Worst case this is O(n^2) - usually a thought O(cn) with small c)...
+			// Not much we without rewriting SListView to use something other than TArray
+			RemoveMessageIdToLogEntry(FilteredLogs[i]);
+			FilteredLogs.RemoveAt(i);
+			--i;
 		}
-	});
-
-	OnLogListChanged().Broadcast(GetFilteredLogs());
-}
+	}
 
-void FFilteredConcertLogList::RemoveMessageIdToLogEntry(const TSharedPtr<FConcertLogEntry>& Log)
-{
-	TArray<TSharedPtr<FConcertLogEntry>>* Entries = MessageIdToLogs.Find(Log->Log.MessageId);
-	if (!Entries)
+	void FFilteredConcertLogList::OnNewLogEntryAdded(const TSharedRef<FConcertLogEntry>& NewLogEntry)
 	{
-		return;
+		if (!Filter || Filter->PassesFilter(*NewLogEntry))
+		{
+			MessageIdToLogs.FindOrAdd(NewLogEntry->Log.MessageId).AddUnique(NewLogEntry);
+			FilteredLogs.Emplace(NewLogEntry);
+			OnLogListChanged().Broadcast(GetFilteredLogs());
+		}
 	}
-	
-	if (Entries->Num() == 1)
+
+	void FFilteredConcertLogList::RebuildFilteredResult()
 	{
-		MessageIdToLogs.Remove(Log->Log.MessageId);
+		MessageIdToLogs.Reset();
+		FilteredLogs.Reset();
+		LogSource->ForEachLog([this](const TSharedPtr<FConcertLogEntry>& LogEntry)
+		{
+			if (!Filter || Filter->PassesFilter(*LogEntry))
+			{
+				MessageIdToLogs.FindOrAdd(LogEntry->Log.MessageId).AddUnique(LogEntry);
+				FilteredLogs.Emplace(LogEntry);
+			}
+		});
+
+		OnLogListChanged().Broadcast(GetFilteredLogs());
 	}
-	else
+
+	void FFilteredConcertLogList::RemoveMessageIdToLogEntry(const TSharedPtr<FConcertLogEntry>& Log)
 	{
-		Entries->RemoveSingle(Log);
+		TArray<TSharedPtr<FConcertLogEntry>>* Entries = MessageIdToLogs.Find(Log->Log.MessageId);
+		if (!Entries)
+		{
+			return;
+		}
+		
+		if (Entries->Num() == 1)
+		{
+			MessageIdToLogs.Remove(Log->Log.MessageId);
+		}
+		else
+		{
+			Entries->RemoveSingle(Log);
+		}
 	}
-}
 
-FPagedFilteredConcertLogList::FPagedFilteredConcertLogList(TSharedRef<IConcertLogSource> LogSource, TSharedPtr<FConcertLogFilter> OptionalFilter, uint16 InitialLogsPerPage)
-	: FFilteredConcertLogList(MoveTemp(LogSource), MoveTemp(OptionalFilter))
-	, LogsPerPage(FMath::Max<uint16>(InitialLogsPerPage, 1))
-{
-	checkf(InitialLogsPerPage >= 1, TEXT("Unreasonable page size"));
-
-	OnLogListChanged().AddRaw(this, &FPagedFilteredConcertLogList::HandeOnLogListChanged);
-	CheckAndConditionallyPopulatePage();
-}
-
-TOptional<FPagedFilteredConcertLogList::FPageCount> FPagedFilteredConcertLogList::GetPageOf(const size_t ItemIndex) const
-{
-	const FPageCount Page = ItemIndex / GetLogsPerPage();
-	if (Page < GetNumPages())
+	FPagedFilteredConcertLogList::FPagedFilteredConcertLogList(TSharedRef<IConcertLogSource> LogSource, TSharedPtr<FConcertLogFilter> OptionalFilter, uint16 InitialLogsPerPage)
+		: FFilteredConcertLogList(MoveTemp(LogSource), MoveTemp(OptionalFilter))
+		, LogsPerPage(FMath::Max<uint16>(InitialLogsPerPage, 1))
 	{
-		return Page;
+		checkf(InitialLogsPerPage >= 1, TEXT("Unreasonable page size"));
+
+		OnLogListChanged().AddRaw(this, &FPagedFilteredConcertLogList::HandeOnLogListChanged);
+		CheckAndConditionallyPopulatePage();
 	}
-	return {};
-}
 
-void FPagedFilteredConcertLogList::SetLogsPerPage(uint16 NewLogsPerPage)
-{
-	if (LogsPerPage != NewLogsPerPage && ensure(NewLogsPerPage > 0))
+	TOptional<FPagedFilteredConcertLogList::FPageCount> FPagedFilteredConcertLogList::GetPageOf(const size_t ItemIndex) const
 	{
-		size_t StartIndex, LastIndex;
-		Tie(StartIndex, LastIndex) = GetLogIndicesForPage();
-		
-		LogsPerPage = NewLogsPerPage;
-		// Recalculate current page so the first item that was the page is still on the page
-		CurrentPageIndex = StartIndex / LogsPerPage;
-
-		RepopulatePage();
+		const FPageCount Page = ItemIndex / GetLogsPerPage();
+		if (Page < GetNumPages())
+		{
+			return Page;
+		}
+		return {};
 	}
-}
 
-void FPagedFilteredConcertLogList::SetPage(const uint32 PageIndex)
-{
-	if (CurrentPageIndex != PageIndex && PageIndex < GetNumPages())
+	void FPagedFilteredConcertLogList::SetLogsPerPage(uint16 NewLogsPerPage)
 	{
-		CurrentPageIndex = PageIndex;
-		RepopulatePage();
+		if (LogsPerPage != NewLogsPerPage && ensure(NewLogsPerPage > 0))
+		{
+			size_t StartIndex, LastIndex;
+			Tie(StartIndex, LastIndex) = GetLogIndicesForPage();
+			
+			LogsPerPage = NewLogsPerPage;
+			// Recalculate current page so the first item that was the page is still on the page
+			CurrentPageIndex = StartIndex / LogsPerPage;
+
+			RepopulatePage();
+		}
 	}
-}
 
-void FPagedFilteredConcertLogList::HandeOnLogListChanged(const TArray<TSharedPtr<FConcertLogEntry>>& NewFilteredLogList)
-{
-	// After filtering there may be less logs > less pages > the current page might have become out of bounds
-	if (CurrentPageIndex >= GetNumPages())
+	void FPagedFilteredConcertLogList::SetPage(const uint32 PageIndex)
 	{
-		CurrentPageIndex = 0;
-		RepopulatePage();
+		if (CurrentPageIndex != PageIndex && PageIndex < GetNumPages())
+		{
+			CurrentPageIndex = PageIndex;
+			RepopulatePage();
+		}
 	}
-	else
+
+	void FPagedFilteredConcertLogList::HandeOnLogListChanged(const TArray<TSharedPtr<FConcertLogEntry>>& NewFilteredLogList)
 	{
-		CheckAndConditionallyPopulatePage();
+		// After filtering there may be less logs > less pages > the current page might have become out of bounds
+		if (CurrentPageIndex >= GetNumPages())
+		{
+			CurrentPageIndex = 0;
+			RepopulatePage();
+		}
+		else
+		{
+			CheckAndConditionallyPopulatePage();
+		}
 	}
-}
 
-void FPagedFilteredConcertLogList::RepopulatePage()
-{
-	PageView.Empty(LogsPerPage);
-	ForEachLogIndexOnPage([this](size_t i)
+	void FPagedFilteredConcertLogList::RepopulatePage()
 	{
-		PageView.Emplace(GetFilteredLogs()[i]);
-	});
+		PageView.Empty(LogsPerPage);
+		ForEachLogIndexOnPage([this](size_t i)
+		{
+			PageView.Emplace(GetFilteredLogs()[i]);
+		});
 
-	OnPageViewChanged().Broadcast(PageView);
-}
+		OnPageViewChanged().Broadcast(PageView);
+	}
 
-void FPagedFilteredConcertLogList::CheckAndConditionallyPopulatePage()
-{
-	bool bChanged = false;
-	size_t PageViewIndex = 0;
-	ForEachLogIndexOnPage([this, &PageViewIndex, &bChanged](size_t i)
+	void FPagedFilteredConcertLogList::CheckAndConditionallyPopulatePage()
 	{
-		if (PageView[PageViewIndex] != GetFilteredLogs()[i])
+		bool bChanged = false;
+		size_t PageViewIndex = 0;
+		ForEachLogIndexOnPage([this, &PageViewIndex, &bChanged](size_t i)
 		{
-			PageView[PageViewIndex] = GetFilteredLogs()[i];
-			bChanged = true;
-		}
+			if (PageView[PageViewIndex] != GetFilteredLogs()[i])
+			{
+				PageView[PageViewIndex] = GetFilteredLogs()[i];
+				bChanged = true;
+			}
 
-		++PageViewIndex;
-	}, FMath::Min(PageView.Num(), GetFilteredLogs().Num()));
+			++PageViewIndex;
+		}, FMath::Min(PageView.Num(), GetFilteredLogs().Num()));
 
-	size_t StartIndex, ItemsOnPage;
-	Tie(StartIndex, ItemsOnPage) = GetLogIndicesForPage();
+		size_t StartIndex, ItemsOnPage;
+		Tie(StartIndex, ItemsOnPage) = GetLogIndicesForPage();
 
-	// Add items that are new in GetFilteredLogs()
-	bChanged |= PageView.Num() < ItemsOnPage; 
-	for (size_t i = PageView.Num(); i < ItemsOnPage && GetFilteredLogs().IsValidIndex(i); ++i)
-	{
-		PageView.Emplace(GetFilteredLogs()[i]);
-	}
+		// Add items that are new in GetFilteredLogs()
+		bChanged |= PageView.Num() < ItemsOnPage; 
+		for (size_t i = PageView.Num(); i < ItemsOnPage && GetFilteredLogs().IsValidIndex(i); ++i)
+		{
+			PageView.Emplace(GetFilteredLogs()[i]);
+		}
 
-	// Remove (possibly duplicate) items that exceed the expected page length 
-	const size_t NumToRemove = PageView.Num() - ItemsOnPage;
-	bChanged |= NumToRemove > 0;
-	for (size_t i = 0; i < NumToRemove; ++i)
-	{
-		PageView.RemoveAt(PageView.Num() - 1);
+		// Remove (possibly duplicate) items that exceed the expected page length 
+		const size_t NumToRemove = PageView.Num() - ItemsOnPage;
+		bChanged |= NumToRemove > 0;
+		for (size_t i = 0; i < NumToRemove; ++i)
+		{
+			PageView.RemoveAt(PageView.Num() - 1);
+		}
+
+		if (bChanged)
+		{
+			OnPageViewChanged().Broadcast(PageView);
+		}
 	}
 
-	if (bChanged)
+	void FPagedFilteredConcertLogList::ForEachLogIndexOnPage(TFunctionRef<void(size_t Index)> Callback, FLogsPerPageCount MaxItems)
 	{
-		OnPageViewChanged().Broadcast(PageView);
+		size_t StartIndex, NumItems;
+		Tie(StartIndex, NumItems) = GetLogIndicesForPage();
+		for (size_t i = 0; i < NumItems && i < MaxItems; ++i)
+		{
+			Callback(StartIndex + i);
+		}
 	}
-}
 
-void FPagedFilteredConcertLogList::ForEachLogIndexOnPage(TFunctionRef<void(size_t Index)> Callback, FLogsPerPageCount MaxItems)
-{
-	size_t StartIndex, NumItems;
-	Tie(StartIndex, NumItems) = GetLogIndicesForPage();
-	for (size_t i = 0; i < NumItems && i < MaxItems; ++i)
+	TTuple<size_t, size_t> FPagedFilteredConcertLogList::GetLogIndicesForPage() const
 	{
-		Callback(StartIndex + i);
+		const size_t StartIndex = CurrentPageIndex * LogsPerPage;
+		const size_t NumberItems = FMath::Min<size_t>(
+			// Used when page is full
+			LogsPerPage,
+			// Used when last page is not full
+			GetFilteredLogs().Num() - StartIndex
+			);
+		check(GetFilteredLogs().Num() == 0 || GetFilteredLogs().IsValidIndex(StartIndex));
+		return { StartIndex, NumberItems };
 	}
 }
-
-TTuple<size_t, size_t> FPagedFilteredConcertLogList::GetLogIndicesForPage() const
-{
-	const size_t StartIndex = CurrentPageIndex * LogsPerPage;
-	const size_t NumberItems = FMath::Min<size_t>(
-		// Used when page is full
-		LogsPerPage,
-		// Used when last page is not full
-		GetFilteredLogs().Num() - StartIndex
-		);
-	check(GetFilteredLogs().Num() == 0 || GetFilteredLogs().IsValidIndex(StartIndex));
-	return { StartIndex, NumberItems };
-}
diff --git a/Engine/Plugins/Developer/Concert/ConcertApp/MultiUserServer/Source/MultiUserServer/Private/Widgets/Clients/Logging/Filter/FilteredConcertLogList.h b/Engine/Plugins/Developer/Concert/ConcertApp/MultiUserServer/Source/MultiUserServer/Private/Widgets/Clients/Logging/Filter/FilteredConcertLogList.h
--- a/Engine/Plugins/Developer/Concert/ConcertApp/MultiUserServer/Source/MultiUserServer/Private/Widgets/Clients/Logging/Filter/FilteredConcertLogList.h
+++ b/Engine/Plugins/Developer/Concert/ConcertApp/MultiUserServer/Source/MultiUserServer/Private/Widgets/Clients/Logging/Filter/FilteredConcertLogList.h
@@ -3,103 +3,105 @@
 #pragma once
 
 #include "CoreMinimal.h"
+#include "ConcertLogFilterTypes.h"
 #include "Widgets/Clients/Logging/Source/IConcertLogSource.h"
 
-class FConcertLogFilter;
 struct FConcertLog;
 
-/** Applies a filter to a IConcertLogSource as the logs come in. */
-class FFilteredConcertLogList
+namespace UE::MultiUserServer
 {
-public:
-	
-	FFilteredConcertLogList(TSharedRef<IConcertLogSource> InLogSource, TSharedPtr<FConcertLogFilter> InOptionalFilter);
-	~FFilteredConcertLogList();
-	
-	DECLARE_MULTICAST_DELEGATE_OneParam(FOnLogListChanged, const TArray<TSharedPtr<FConcertLogEntry>>& /*NewFilteredView*/);
-	FOnLogListChanged& OnLogListChanged() { return OnLogListChangedEvent; }
-
-	bool IsBufferFull() const { return LogSource->IsBufferFull(); }
-	const TArray<TSharedPtr<FConcertLogEntry>>& GetFilteredLogs() const { return FilteredLogs; }
-	TArray<TSharedPtr<FConcertLogEntry>> GetFilteredLogsWithId(const FGuid& MessageId) const
+	/** Applies a filter to a IConcertLogSource as the logs come in. */
+	class FFilteredConcertLogList
 	{
-		const TArray<TSharedPtr<FConcertLogEntry>>* Result = MessageIdToLogs.Find(MessageId);
-		return Result ? *Result : TArray<TSharedPtr<FConcertLogEntry>>{};
-	}
-
-protected:
-
-	const TSharedRef<IConcertLogSource>& GetLogSource() const { return LogSource; }
-	
-private:
-
-	/** Determines where the logs come from */
-	TSharedRef<IConcertLogSource> LogSource;
-
-	/** Optional. Filters items from the LogSource */
-	TSharedPtr<FConcertLogFilter> Filter;
-
-	/** Contains the items of LogSource with Filter applied. */
-	TArray<TSharedPtr<FConcertLogEntry>> FilteredLogs;
-
-	/** Binds message IDs to the logs that have that ID. Logs can share the same ID but have different MessageActions. */
-	TMap<FGuid, TArray<TSharedPtr<FConcertLogEntry>>> MessageIdToLogs;
-	
-	/** Called when the content of FilteredResult changes */
-	FOnLogListChanged OnLogListChangedEvent;
-
-	void OnLowestLogEntryChanged(FConcertLogID NewLowestValidID);
-	void OnNewLogEntryAdded(const TSharedRef<FConcertLogEntry>& NewLogEntry);
-	void RebuildFilteredResult();
-
-	void RemoveMessageIdToLogEntry(const TSharedPtr<FConcertLogEntry>& Log);
-};
-
-/** Further sorts the results of FFilteredConcertLogList::Filter into pages, i.e. buckets of logs, useful for limiting elements in a table view. */
-class FPagedFilteredConcertLogList : public FFilteredConcertLogList
-{
-public:
-	
-	using FLogsPerPageCount = uint16;
-	using FPageCount = uint32;
-
-	FPagedFilteredConcertLogList(TSharedRef<IConcertLogSource> LogSource, TSharedPtr<FConcertLogFilter> OptionalFilter, uint16 FLogsPerPageCount = 500);
-
-	FPageCount GetNumPages() const { const FPageCount PageCount = GetFilteredLogs().Num() / LogsPerPage; return GetFilteredLogs().Num() % LogsPerPage == 0 ? PageCount : PageCount + 1;}
-	FPageCount GetCurrentPage() const { return CurrentPageIndex; }
-	TOptional<FPageCount> GetPageOf(const size_t ItemIndex) const;
-	FLogsPerPageCount GetLogsPerPage() const { return LogsPerPage; }
-	const TArray<TSharedPtr<FConcertLogEntry>>& GetPageView() const { return PageView; }
-	
-	void SetLogsPerPage(FLogsPerPageCount NewLogsPerPage);
-	void SetPage(const FPageCount PageIndex);
-
-	FOnLogListChanged& OnPageViewChanged() { return OnPageViewChangedEvent; }
-	
-private:
-
-
-	/** The number of logs to show per page*/
-	FLogsPerPageCount LogsPerPage;
-
-	/** The index of the current page. Starts at 0. */
-	FPageCount CurrentPageIndex = 0;
-
-	/** The filtered logs on the current page */
-	TArray<TSharedPtr<FConcertLogEntry>> PageView;
-
-	/** Called when the content of PageView changes */
-	FOnLogListChanged OnPageViewChangedEvent;
-
-	void HandeOnLogListChanged(const TArray<TSharedPtr<FConcertLogEntry>>& NewFilteredLogList);
-	void RepopulatePage();
-	void CheckAndConditionallyPopulatePage();
-	
-	/**
-	 * Given the current page settings (not validated - assumed corrected), calls Callback with each index in GetFilteredLogs()
-	 * that should be an item on the current page. */
-	void ForEachLogIndexOnPage(TFunctionRef<void(size_t FilteredLogIndex)> Callback, FLogsPerPageCount MaxItems = TNumericLimits<FLogsPerPageCount>::Max());
-	/** @return Start and last index in GetFilteredLogs() array to accommodate current page settings */
-	TTuple<size_t, size_t> GetLogIndicesForPage() const;
-};
-
+	public:
+		
+		FFilteredConcertLogList(TSharedRef<IConcertLogSource> InLogSource, TSharedPtr<FConcertLogFilter> InOptionalFilter);
+		~FFilteredConcertLogList();
+		
+		DECLARE_MULTICAST_DELEGATE_OneParam(FOnLogListChanged, const TArray<TSharedPtr<FConcertLogEntry>>& /*NewFilteredView*/);
+		FOnLogListChanged& OnLogListChanged() { return OnLogListChangedEvent; }
+
+		bool IsBufferFull() const { return LogSource->IsBufferFull(); }
+		const TArray<TSharedPtr<FConcertLogEntry>>& GetFilteredLogs() const { return FilteredLogs; }
+		TArray<TSharedPtr<FConcertLogEntry>> GetFilteredLogsWithId(const FGuid& MessageId) const
+		{
+			const TArray<TSharedPtr<FConcertLogEntry>>* Result = MessageIdToLogs.Find(MessageId);
+			return Result ? *Result : TArray<TSharedPtr<FConcertLogEntry>>{};
+		}
+
+	protected:
+
+		const TSharedRef<IConcertLogSource>& GetLogSource() const { return LogSource; }
+		
+	private:
+
+		/** Determines where the logs come from */
+		TSharedRef<IConcertLogSource> LogSource;
+
+		/** Optional. Filters items from the LogSource */
+		TSharedPtr<FConcertLogFilter> Filter;
+
+		/** Contains the items of LogSource with Filter applied. */
+		TArray<TSharedPtr<FConcertLogEntry>> FilteredLogs;
+
+		/** Binds message IDs to the logs that have that ID. Logs can share the same ID but have different MessageActions. */
+		TMap<FGuid, TArray<TSharedPtr<FConcertLogEntry>>> MessageIdToLogs;
+		
+		/** Called when the content of FilteredResult changes */
+		FOnLogListChanged OnLogListChangedEvent;
+
+		void OnLowestLogEntryChanged(FConcertLogID NewLowestValidID);
+		void OnNewLogEntryAdded(const TSharedRef<FConcertLogEntry>& NewLogEntry);
+		void RebuildFilteredResult();
+
+		void RemoveMessageIdToLogEntry(const TSharedPtr<FConcertLogEntry>& Log);
+	};
+
+	/** Further sorts the results of FFilteredConcertLogList::Filter into pages, i.e. buckets of logs, useful for limiting elements in a table view. */
+	class FPagedFilteredConcertLogList : public FFilteredConcertLogList
+	{
+	public:
+		
+		using FLogsPerPageCount = uint16;
+		using FPageCount = uint32;
+
+		FPagedFilteredConcertLogList(TSharedRef<IConcertLogSource> LogSource, TSharedPtr<FConcertLogFilter> OptionalFilter, uint16 FLogsPerPageCount = 500);
+
+		FPageCount GetNumPages() const { const FPageCount PageCount = GetFilteredLogs().Num() / LogsPerPage; return GetFilteredLogs().Num() % LogsPerPage == 0 ? PageCount : PageCount + 1;}
+		FPageCount GetCurrentPage() const { return CurrentPageIndex; }
+		TOptional<FPageCount> GetPageOf(const size_t ItemIndex) const;
+		FLogsPerPageCount GetLogsPerPage() const { return LogsPerPage; }
+		const TArray<TSharedPtr<FConcertLogEntry>>& GetPageView() const { return PageView; }
+		
+		void SetLogsPerPage(FLogsPerPageCount NewLogsPerPage);
+		void SetPage(const FPageCount PageIndex);
+
+		FOnLogListChanged& OnPageViewChanged() { return OnPageViewChangedEvent; }
+		
+	private:
+
+
+		/** The number of logs to show per page*/
+		FLogsPerPageCount LogsPerPage;
+
+		/** The index of the current page. Starts at 0. */
+		FPageCount CurrentPageIndex = 0;
+
+		/** The filtered logs on the current page */
+		TArray<TSharedPtr<FConcertLogEntry>> PageView;
+
+		/** Called when the content of PageView changes */
+		FOnLogListChanged OnPageViewChangedEvent;
+
+		void HandeOnLogListChanged(const TArray<TSharedPtr<FConcertLogEntry>>& NewFilteredLogList);
+		void RepopulatePage();
+		void CheckAndConditionallyPopulatePage();
+		
+		/**
+		 * Given the current page settings (not validated - assumed corrected), calls Callback with each index in GetFilteredLogs()
+		 * that should be an item on the current page. */
+		void ForEachLogIndexOnPage(TFunctionRef<void(size_t FilteredLogIndex)> Callback, FLogsPerPageCount MaxItems = TNumericLimits<FLogsPerPageCount>::Max());
+		/** @return Start and last index in GetFilteredLogs() array to accommodate current page settings */
+		TTuple<size_t, size_t> GetLogIndicesForPage() const;
+	};
+}
diff --git a/Engine/Plugins/Developer/Concert/ConcertApp/MultiUserServer/Source/MultiUserServer/Private/Widgets/Clients/Logging/SConcertTransportLog.cpp b/Engine/Plugins/Developer/Concert/ConcertApp/MultiUserServer/Source/MultiUserServer/Private/Widgets/Clients/Logging/SConcertTransportLog.cpp
--- a/Engine/Plugins/Developer/Concert/ConcertApp/MultiUserServer/Source/MultiUserServer/Private/Widgets/Clients/Logging/SConcertTransportLog.cpp
+++ b/Engine/Plugins/Developer/Concert/ConcertApp/MultiUserServer/Source/MultiUserServer/Private/Widgets/Clients/Logging/SConcertTransportLog.cpp
@@ -23,392 +23,395 @@
 
 #define LOCTEXT_NAMESPACE "UnrealMultiUserUI.SConcertTransportLog"
 
-const FName SConcertTransportLog::FirstColumnId("AvatarColourColumnId");
-
-SConcertTransportLog::~SConcertTransportLog()
+namespace UE::MultiUserServer
 {
-	ConcertTransportEvents::OnConcertTransportLoggingEnabledChangedEvent().RemoveAll(this);
-	UMultiUserServerColumnVisibilitySettings::GetSettings()->OnTransportLogColumnVisibility().RemoveAll(this);
-}
-
-void SConcertTransportLog::Construct(const FArguments& InArgs, TSharedRef<IConcertLogSource> LogSource, TSharedRef<FEndpointToUserNameCache> InEndpointCache, TSharedRef<FConcertLogTokenizer> InLogTokenizer)
-{
-	PagedLogList = MakeShared<FPagedFilteredConcertLogList>(MoveTemp(LogSource), InArgs._Filter);
-	EndpointCache = MoveTemp(InEndpointCache);
-	LogTokenizer = MoveTemp(InLogTokenizer);
-	Filter = InArgs._Filter;
+	const FName SConcertTransportLog::FirstColumnId("AvatarColourColumnId");
 	
-	ChildSlot
-	[
-		SAssignNew(EnableLoggingPromptOverlay, SOverlay)
+	SConcertTransportLog::~SConcertTransportLog()
+	{
+		ConcertTransportEvents::OnConcertTransportLoggingEnabledChangedEvent().RemoveAll(this);
+		UMultiUserServerColumnVisibilitySettings::GetSettings()->OnTransportLogColumnVisibility().RemoveAll(this);
+	}
 
-		+SOverlay::Slot()
+	void SConcertTransportLog::Construct(const FArguments& InArgs, TSharedRef<IConcertLogSource> LogSource, TSharedRef<FEndpointToUserNameCache> InEndpointCache, TSharedRef<FConcertLogTokenizer> InLogTokenizer)
+	{
+		PagedLogList = MakeShared<FPagedFilteredConcertLogList>(MoveTemp(LogSource), InArgs._Filter);
+		EndpointCache = MoveTemp(InEndpointCache);
+		LogTokenizer = MoveTemp(InLogTokenizer);
+		Filter = InArgs._Filter;
+		
+		ChildSlot
 		[
-			SNew(SBorder)
-			.BorderImage(FAppStyle::Get().GetBrush("ToolPanel.GroupBorder"))
-			.BorderBackgroundColor(FSlateColor(FLinearColor(0.6, 0.6, 0.6)))
-			.Padding(2)
-			[
-				SNew(SVerticalBox)
+			SAssignNew(EnableLoggingPromptOverlay, SOverlay)
 
-				+SVerticalBox::Slot()
-				.AutoHeight()
-				.VAlign(VAlign_Top)
-				[
-					InArgs._Filter ? InArgs._Filter->BuildFilterWidgets() : SNullWidget::NullWidget
-				]
-
-				+SVerticalBox::Slot()
-				.FillHeight(1.f)
-				.Padding(0, 5, 0, 0)
-				[
-					CreateTableView()
-				]
-
-				+SVerticalBox::Slot()
-				.AutoHeight()
+			+SOverlay::Slot()
+			[
+				SNew(SBorder)
+				.BorderImage(FAppStyle::Get().GetBrush("ToolPanel.GroupBorder"))
+				.BorderBackgroundColor(FSlateColor(FLinearColor(0.6, 0.6, 0.6)))
+				.Padding(2)
 				[
-					SNew(SConcertTransportLogFooter, PagedLogList.ToSharedRef())
-					.ExtendViewOptions(this, &SConcertTransportLog::ExtendViewOptions)
+					SNew(SVerticalBox)
+
+					+SVerticalBox::Slot()
+					.AutoHeight()
+					.VAlign(VAlign_Top)
+					[
+						InArgs._Filter ? InArgs._Filter->BuildFilterWidgets() : SNullWidget::NullWidget
+					]
+
+					+SVerticalBox::Slot()
+					.FillHeight(1.f)
+					.Padding(0, 5, 0, 0)
+					[
+						CreateTableView()
+					]
+
+					+SVerticalBox::Slot()
+					.AutoHeight()
+					[
+						SNew(SConcertTransportLogFooter, PagedLogList.ToSharedRef())
+						.ExtendViewOptions(this, &SConcertTransportLog::ExtendViewOptions)
+					]
 				]
 			]
-		]
-	];
+		];
 
-	// Subscribe to events
-	PagedLogList->OnPageViewChanged().AddSP(this, &SConcertTransportLog::OnPageViewChanged);
-	
-	UMultiUserServerColumnVisibilitySettings::GetSettings()->OnTransportLogColumnVisibility().AddSP(this, &SConcertTransportLog::OnColumnVisibilitySettingsChanged);
-	UE::ConcertSharedSlate::RestoreColumnVisibilityState(HeaderRow.ToSharedRef(), UMultiUserServerColumnVisibilitySettings::GetSettings()->GetTransportLogColumnVisibility());
-	
-	ConcertTransportEvents::OnConcertTransportLoggingEnabledChangedEvent().AddSP(this, &SConcertTransportLog::OnConcertLoggingEnabledChanged);
-	OnConcertLoggingEnabledChanged(ConcertTransportEvents::IsLoggingEnabled());
-}
+		// Subscribe to events
+		PagedLogList->OnPageViewChanged().AddSP(this, &SConcertTransportLog::OnPageViewChanged);
+		
+		UMultiUserServerColumnVisibilitySettings::GetSettings()->OnTransportLogColumnVisibility().AddSP(this, &SConcertTransportLog::OnColumnVisibilitySettingsChanged);
+		UE::ConcertSharedSlate::RestoreColumnVisibilityState(HeaderRow.ToSharedRef(), UMultiUserServerColumnVisibilitySettings::GetSettings()->GetTransportLogColumnVisibility());
+		
+		ConcertTransportEvents::OnConcertTransportLoggingEnabledChangedEvent().AddSP(this, &SConcertTransportLog::OnConcertLoggingEnabledChanged);
+		OnConcertLoggingEnabledChanged(ConcertTransportEvents::IsLoggingEnabled());
+	}
 
-bool SConcertTransportLog::CanScrollToLog(const FGuid& MessageId, FConcertLogEntryFilterFunc FilterFunc) const
-{
-	return PagedLogList->GetFilteredLogsWithId(MessageId).ContainsByPredicate([FilterFunc](const TSharedPtr<FConcertLogEntry>& LogEntry)
+	bool SConcertTransportLog::CanScrollToLog(const FGuid& MessageId, FConcertLogEntryFilterFunc FilterFunc) const
 	{
-		return FilterFunc(*LogEntry);
-	});
-}
+		return PagedLogList->GetFilteredLogsWithId(MessageId).ContainsByPredicate([FilterFunc](const TSharedPtr<FConcertLogEntry>& LogEntry)
+		{
+			return FilterFunc(*LogEntry);
+		});
+	}
 
-void SConcertTransportLog::ScrollToLog(const FGuid& MessageId, FConcertLogEntryFilterFunc FilterFunc) const
-{
-	// Could be optimised further by having FPagedFilteredConcertLogList cache the index but it's fine - ScrollToLog is called very infrequently (when the user presses a button)
-	const int32 Index = PagedLogList->GetFilteredLogs().IndexOfByPredicate([MessageId, FilterFunc](const TSharedPtr<FConcertLogEntry>& Entry)
+	void SConcertTransportLog::ScrollToLog(const FGuid& MessageId, FConcertLogEntryFilterFunc FilterFunc) const
 	{
-		return Entry->Log.MessageId == MessageId && FilterFunc(*Entry);
-	});
-	ScrollToLog(Index);
-}
+		// Could be optimised further by having FPagedFilteredConcertLogList cache the index but it's fine - ScrollToLog is called very infrequently (when the user presses a button)
+		const int32 Index = PagedLogList->GetFilteredLogs().IndexOfByPredicate([MessageId, FilterFunc](const TSharedPtr<FConcertLogEntry>& Entry)
+		{
+			return Entry->Log.MessageId == MessageId && FilterFunc(*Entry);
+		});
+		ScrollToLog(Index);
+	}
 
-TSharedRef<SWidget> SConcertTransportLog::CreateTableView()
-{
-	return SAssignNew(LogView, SListView<TSharedPtr<FConcertLogEntry>>)
-		.ListItemsSource(&PagedLogList->GetPageView())
-		.OnGenerateRow(this, &SConcertTransportLog::OnGenerateActivityRowWidget)
-		.SelectionMode(ESelectionMode::Multi)
-		.HeaderRow(CreateHeaderRow());
-}
+	TSharedRef<SWidget> SConcertTransportLog::CreateTableView()
+	{
+		return SAssignNew(LogView, SListView<TSharedPtr<FConcertLogEntry>>)
+			.ListItemsSource(&PagedLogList->GetPageView())
+			.OnGenerateRow(this, &SConcertTransportLog::OnGenerateActivityRowWidget)
+			.SelectionMode(ESelectionMode::Multi)
+			.HeaderRow(CreateHeaderRow());
+	}
 
-TSharedRef<SHeaderRow> SConcertTransportLog::CreateHeaderRow()
-{
-	HeaderRow = SNew(SHeaderRow)
-		.OnHiddenColumnsListChanged_Lambda([this]()
-		{
-			if (!bIsUpdatingColumnVisibility)
-			{
-				UMultiUserServerColumnVisibilitySettings::GetSettings()->SetTransportLogColumnVisibility(
-				UE::ConcertSharedSlate::SnapshotColumnVisibilityState(HeaderRow.ToSharedRef())
-				);
-			}
-		})
-		// Create tiny dummy row showing avatar colour to handle the case when a user hides all columns
-		+SHeaderRow::Column(FirstColumnId)
-			.DefaultLabel(FText::GetEmpty())
-			.FixedWidth(8)
-			// Cannot be hidden
-			.ShouldGenerateWidget(true)
-			.ToolTipText(LOCTEXT("AvatarColumnToolTipText", "The colour of the avatar is affected by log"))
-		;
-	
-	ForEachPropertyColumn([this](const FProperty& Property, FName ColumnId)
+	TSharedRef<SHeaderRow> SConcertTransportLog::CreateHeaderRow()
 	{
-		const TMap<FName, FString> ColumnNameOverrides = {
-			{ GET_MEMBER_NAME_CHECKED(FConcertLog, MessageAction), TEXT("Action") },
-			{ GET_MEMBER_NAME_CHECKED(FConcertLog, MessageOrderIndex), TEXT("Order Index") },
-			{ GET_MEMBER_NAME_CHECKED(FConcertLog, ChannelId), TEXT("Channel") },
-			{ GET_MEMBER_NAME_CHECKED(FConcertLog, CustomPayloadUncompressedByteSize), TEXT("Size") },
-			{ GET_MEMBER_NAME_CHECKED(FConcertLog, OriginEndpointId), TEXT("Origin") },
-			{ GET_MEMBER_NAME_CHECKED(FConcertLog, DestinationEndpointId), TEXT("Destination") },
-			{ GET_MEMBER_NAME_CHECKED(FConcertLogMetadata, AckState), TEXT("Ack") }
-		};
-		const TMap<FName, float> ManuallySizeColumns = {
-			{ GET_MEMBER_NAME_CHECKED(FConcertLog, Frame), 60 },
-			{ GET_MEMBER_NAME_CHECKED(FConcertLog, MessageId), 280 },
-			{ GET_MEMBER_NAME_CHECKED(FConcertLog, MessageOrderIndex), 60 },
-			{ GET_MEMBER_NAME_CHECKED(FConcertLog, ChannelId), 60 },
-			{ GET_MEMBER_NAME_CHECKED(FConcertLog, Timestamp), 160 },
-			{ GET_MEMBER_NAME_CHECKED(FConcertLog, MessageAction), 80 },
-			{ GET_MEMBER_NAME_CHECKED(FConcertLog, MessageTypeName), 200 },
-			{ GET_MEMBER_NAME_CHECKED(FConcertLog, CustomPayloadTypename), 400 },
-			{ GET_MEMBER_NAME_CHECKED(FConcertLog, MessageAction), 80 },
-			{ GET_MEMBER_NAME_CHECKED(FConcertLog, CustomPayloadUncompressedByteSize), 100 },
-			{ GET_MEMBER_NAME_CHECKED(FConcertLogMetadata, AckState), 45 }
-		};
-		
-		const FString* LabelOverride = ColumnNameOverrides.Find(ColumnId);
-		const FString PropertyName = LabelOverride ? *LabelOverride : Property.GetAuthoredName();
-		SHeaderRow::FColumn::FArguments Args = SHeaderRow::FColumn::FArguments()
-			.ColumnId(ColumnId)
-			.DefaultLabel(FText::FromString(PropertyName))
-			.HAlignCell(HAlign_Center)
-			// Add option to hide
-			.OnGetMenuContent_Lambda([this, ColumnId]()
+		HeaderRow = SNew(SHeaderRow)
+			.OnHiddenColumnsListChanged_Lambda([this]()
 			{
-				return UE::ConcertSharedSlate::MakeHideColumnContextMenu(HeaderRow.ToSharedRef(), ColumnId);
-			});
+				if (!bIsUpdatingColumnVisibility)
+				{
+					UMultiUserServerColumnVisibilitySettings::GetSettings()->SetTransportLogColumnVisibility(
+					UE::ConcertSharedSlate::SnapshotColumnVisibilityState(HeaderRow.ToSharedRef())
+					);
+				}
+			})
+			// Create tiny dummy row showing avatar colour to handle the case when a user hides all columns
+			+SHeaderRow::Column(FirstColumnId)
+				.DefaultLabel(FText::GetEmpty())
+				.FixedWidth(8)
+				// Cannot be hidden
+				.ShouldGenerateWidget(true)
+				.ToolTipText(LOCTEXT("AvatarColumnToolTipText", "The colour of the avatar is affected by log"))
+			;
 		
-		if (const float* ManualSize = ManuallySizeColumns.Find(ColumnId))
+		ForEachPropertyColumn([this](const FProperty& Property, FName ColumnId)
 		{
-			Args.ManualWidth(*ManualSize);
-		}
-		
-		HeaderRow->AddColumn(Args);
-	});
+			const TMap<FName, FString> ColumnNameOverrides = {
+				{ GET_MEMBER_NAME_CHECKED(FConcertLog, MessageAction), TEXT("Action") },
+				{ GET_MEMBER_NAME_CHECKED(FConcertLog, MessageOrderIndex), TEXT("Order Index") },
+				{ GET_MEMBER_NAME_CHECKED(FConcertLog, ChannelId), TEXT("Channel") },
+				{ GET_MEMBER_NAME_CHECKED(FConcertLog, CustomPayloadUncompressedByteSize), TEXT("Size") },
+				{ GET_MEMBER_NAME_CHECKED(FConcertLog, OriginEndpointId), TEXT("Origin") },
+				{ GET_MEMBER_NAME_CHECKED(FConcertLog, DestinationEndpointId), TEXT("Destination") },
+				{ GET_MEMBER_NAME_CHECKED(FConcertLogMetadata, AckState), TEXT("Ack") }
+			};
+			const TMap<FName, float> ManuallySizeColumns = {
+				{ GET_MEMBER_NAME_CHECKED(FConcertLog, Frame), 60 },
+				{ GET_MEMBER_NAME_CHECKED(FConcertLog, MessageId), 280 },
+				{ GET_MEMBER_NAME_CHECKED(FConcertLog, MessageOrderIndex), 60 },
+				{ GET_MEMBER_NAME_CHECKED(FConcertLog, ChannelId), 60 },
+				{ GET_MEMBER_NAME_CHECKED(FConcertLog, Timestamp), 160 },
+				{ GET_MEMBER_NAME_CHECKED(FConcertLog, MessageAction), 80 },
+				{ GET_MEMBER_NAME_CHECKED(FConcertLog, MessageTypeName), 200 },
+				{ GET_MEMBER_NAME_CHECKED(FConcertLog, CustomPayloadTypename), 400 },
+				{ GET_MEMBER_NAME_CHECKED(FConcertLog, MessageAction), 80 },
+				{ GET_MEMBER_NAME_CHECKED(FConcertLog, CustomPayloadUncompressedByteSize), 100 },
+				{ GET_MEMBER_NAME_CHECKED(FConcertLogMetadata, AckState), 45 }
+			};
+			
+			const FString* LabelOverride = ColumnNameOverrides.Find(ColumnId);
+			const FString PropertyName = LabelOverride ? *LabelOverride : Property.GetAuthoredName();
+			SHeaderRow::FColumn::FArguments Args = SHeaderRow::FColumn::FArguments()
+				.ColumnId(ColumnId)
+				.DefaultLabel(FText::FromString(PropertyName))
+				.HAlignCell(HAlign_Center)
+				// Add option to hide
+				.OnGetMenuContent_Lambda([this, ColumnId]()
+				{
+					return UE::ConcertSharedSlate::MakeHideColumnContextMenu(HeaderRow.ToSharedRef(), ColumnId);
+				});
+			
+			if (const float* ManualSize = ManuallySizeColumns.Find(ColumnId))
+			{
+				Args.ManualWidth(*ManualSize);
+			}
+			
+			HeaderRow->AddColumn(Args);
+		});
 
-	TGuardValue<bool> DoNotSave(bIsUpdatingColumnVisibility, true);
-	RestoreDefaultColumnVisiblities();
-	
-	return HeaderRow.ToSharedRef();
-}
+		TGuardValue<bool> DoNotSave(bIsUpdatingColumnVisibility, true);
+		RestoreDefaultColumnVisiblities();
+		
+		return HeaderRow.ToSharedRef();
+	}
 
-TSharedRef<ITableRow> SConcertTransportLog::OnGenerateActivityRowWidget(TSharedPtr<FConcertLogEntry> Item, const TSharedRef<STableViewBase>& OwnerTable) const
-{
-	const TOptional<FConcertClientInfo> OriginInfo = EndpointCache->GetClientInfo(Item->Log.OriginEndpointId);
-	const TOptional<FConcertClientInfo> DestinationInfo = EndpointCache->GetClientInfo(Item->Log.DestinationEndpointId);
-	const FLinearColor AvatarColor = OriginInfo.IsSet()
-		? OriginInfo->AvatarColor
-		: DestinationInfo.IsSet() ? DestinationInfo->AvatarColor : FLinearColor::Black;
-	
-	return SNew(SConcertTransportLogRow, Item, OwnerTable, LogTokenizer.ToSharedRef())
-		.HighlightText_Lambda([this](){ return Filter.IsValid() ? Filter->GetTextSearchFilter()->GetSearchText() : FText(); })
-		.AvatarColor(AvatarColor)
-		.CanScrollToAckLog(this, &SConcertTransportLog::CanScrollToAckLog)
-		.CanScrollToAckedLog(this, &SConcertTransportLog::CanScrollToAckedLog)
-		.ScrollToAckLog(this, &SConcertTransportLog::ScrollToAckLog)
-		.ScrollToAckedLog(this, &SConcertTransportLog::ScrollToAckedLog);
-}
+	TSharedRef<ITableRow> SConcertTransportLog::OnGenerateActivityRowWidget(TSharedPtr<FConcertLogEntry> Item, const TSharedRef<STableViewBase>& OwnerTable) const
+	{
+		const TOptional<FConcertClientInfo> OriginInfo = EndpointCache->GetClientInfo(Item->Log.OriginEndpointId);
+		const TOptional<FConcertClientInfo> DestinationInfo = EndpointCache->GetClientInfo(Item->Log.DestinationEndpointId);
+		const FLinearColor AvatarColor = OriginInfo.IsSet()
+			? OriginInfo->AvatarColor
+			: DestinationInfo.IsSet() ? DestinationInfo->AvatarColor : FLinearColor::Black;
+		
+		return SNew(SConcertTransportLogRow, Item, OwnerTable, LogTokenizer.ToSharedRef())
+			.HighlightText_Lambda([this](){ return Filter.IsValid() ? Filter->GetTextSearchFilter()->GetSearchText() : FText(); })
+			.AvatarColor(AvatarColor)
+			.CanScrollToAckLog(this, &SConcertTransportLog::CanScrollToAckLog)
+			.CanScrollToAckedLog(this, &SConcertTransportLog::CanScrollToAckedLog)
+			.ScrollToAckLog(this, &SConcertTransportLog::ScrollToAckLog)
+			.ScrollToAckedLog(this, &SConcertTransportLog::ScrollToAckedLog);
+	}
 
-void SConcertTransportLog::ForEachPropertyColumn(TFunctionRef<void(const FProperty& ColumnPropertyId, FName ColumnId)> Callback)
-{
-	for (TFieldIterator<const FProperty> PropertyIt(FConcertLog::StaticStruct()); PropertyIt; ++PropertyIt)
+	void SConcertTransportLog::ForEachPropertyColumn(TFunctionRef<void(const FProperty& ColumnPropertyId, FName ColumnId)> Callback)
 	{
-		if (!PropertyIt->HasAnyPropertyFlags(CPF_Transient))
+		for (TFieldIterator<const FProperty> PropertyIt(FConcertLog::StaticStruct()); PropertyIt; ++PropertyIt)
+		{
+			if (!PropertyIt->HasAnyPropertyFlags(CPF_Transient))
+			{
+				Callback(**PropertyIt, PropertyIt->GetFName());
+			}
+		}
+		for (TFieldIterator<const FProperty> PropertyIt(FConcertLogMetadata::StaticStruct()); PropertyIt; ++PropertyIt)
 		{
 			Callback(**PropertyIt, PropertyIt->GetFName());
 		}
 	}
-	for (TFieldIterator<const FProperty> PropertyIt(FConcertLogMetadata::StaticStruct()); PropertyIt; ++PropertyIt)
+
+	void SConcertTransportLog::RestoreDefaultColumnVisiblities()
 	{
-		Callback(**PropertyIt, PropertyIt->GetFName());
+		ForEachPropertyColumn([this](const FProperty& Property, FName ColumnId)
+		{
+			const TSet<FName> HiddenByDefault = {
+				GET_MEMBER_NAME_CHECKED(FConcertLog, Frame),
+				GET_MEMBER_NAME_CHECKED(FConcertLog, MessageId),
+				GET_MEMBER_NAME_CHECKED(FConcertLog, MessageOrderIndex),
+				GET_MEMBER_NAME_CHECKED(FConcertLog, ChannelId),
+				GET_MEMBER_NAME_CHECKED(FConcertLog, CustomPayloadTypename),
+				GET_MEMBER_NAME_CHECKED(FConcertLog, StringPayload)
+			};
+			HeaderRow->SetShowGeneratedColumn(ColumnId, !HiddenByDefault.Contains(ColumnId));
+		});
 	}
-}
 
-void SConcertTransportLog::RestoreDefaultColumnVisiblities()
-{
-	ForEachPropertyColumn([this](const FProperty& Property, FName ColumnId)
+	void SConcertTransportLog::ExtendViewOptions(FMenuBuilder& MenuBuilder)
 	{
-		const TSet<FName> HiddenByDefault = {
-			GET_MEMBER_NAME_CHECKED(FConcertLog, Frame),
-			GET_MEMBER_NAME_CHECKED(FConcertLog, MessageId),
-			GET_MEMBER_NAME_CHECKED(FConcertLog, MessageOrderIndex),
-			GET_MEMBER_NAME_CHECKED(FConcertLog, ChannelId),
-			GET_MEMBER_NAME_CHECKED(FConcertLog, CustomPayloadTypename),
-			GET_MEMBER_NAME_CHECKED(FConcertLog, StringPayload)
-		};
-		HeaderRow->SetShowGeneratedColumn(ColumnId, !HiddenByDefault.Contains(ColumnId));
-	});
-}
-
-void SConcertTransportLog::ExtendViewOptions(FMenuBuilder& MenuBuilder)
-{
-	MenuBuilder.AddMenuEntry(
-		LOCTEXT("AutoScroll", "Auto Scroll"),
-		LOCTEXT("AutoScroll_Tooltip", "Automatically scroll as new logs arrive (affects last page)"),
-		FSlateIcon(),
-		FUIAction(
-			FExecuteAction::CreateLambda([this](){ bAutoScroll = !bAutoScroll; }),
-			FCanExecuteAction::CreateLambda([] { return true; }),
-			FIsActionChecked::CreateLambda([this] { return bAutoScroll; })),
-		NAME_None,
-		EUserInterfaceActionType::ToggleButton
-	);
-
-	MenuBuilder.AddMenuEntry(
-			LOCTEXT("DisplayTimestampInRelativeTime", "Display Relative Time"),
-			TAttribute<FText>::CreateLambda([this]()
-			{
-				const bool bIsVisible = HeaderRow->IsColumnVisible(GET_MEMBER_NAME_CHECKED(FConcertLog, Timestamp));
-				return bIsVisible
-					? LOCTEXT("DisplayTimestampInRelativeTime.Tooltip.Visible", "Display the Last Modified column in relative time?")
-					: LOCTEXT("DisplayTimestampInRelativeTime.Tooltip.Hidden", "Disabled because the Timestamp column is hidden.");
-			}),
+		MenuBuilder.AddMenuEntry(
+			LOCTEXT("AutoScroll", "Auto Scroll"),
+			LOCTEXT("AutoScroll_Tooltip", "Automatically scroll as new logs arrive (affects last page)"),
 			FSlateIcon(),
 			FUIAction(
-				FExecuteAction::CreateSP(this, &SConcertTransportLog::OnFilterMenuChecked),
-				FCanExecuteAction::CreateLambda([this] { return HeaderRow->IsColumnVisible(GET_MEMBER_NAME_CHECKED(FConcertLog, Timestamp)); }),
-				FIsActionChecked::CreateLambda([this] { return UConcertTransportLogSettings::GetSettings()->TimestampTimeFormat == ETimeFormat::Relative; })),
+				FExecuteAction::CreateLambda([this](){ bAutoScroll = !bAutoScroll; }),
+				FCanExecuteAction::CreateLambda([] { return true; }),
+				FIsActionChecked::CreateLambda([this] { return bAutoScroll; })),
 			NAME_None,
 			EUserInterfaceActionType::ToggleButton
 		);
 
-	MenuBuilder.AddSeparator();
-	MenuBuilder.AddMenuEntry(
-			LOCTEXT("SelectAll", "Show all"),
-			FText::GetEmpty(),
-			FSlateIcon(),
-			FUIAction(
-				FExecuteAction::CreateLambda([this]()
+		MenuBuilder.AddMenuEntry(
+				LOCTEXT("DisplayTimestampInRelativeTime", "Display Relative Time"),
+				TAttribute<FText>::CreateLambda([this]()
 				{
-					ForEachPropertyColumn([this](const FProperty& Property, FName ColumnId)
-					{
-						HeaderRow->SetShowGeneratedColumn(ColumnId, true);
-					});
+					const bool bIsVisible = HeaderRow->IsColumnVisible(GET_MEMBER_NAME_CHECKED(FConcertLog, Timestamp));
+					return bIsVisible
+						? LOCTEXT("DisplayTimestampInRelativeTime.Tooltip.Visible", "Display the Last Modified column in relative time?")
+						: LOCTEXT("DisplayTimestampInRelativeTime.Tooltip.Hidden", "Disabled because the Timestamp column is hidden.");
 				}),
-				FCanExecuteAction::CreateLambda([] { return true; })),
-			NAME_None,
-			EUserInterfaceActionType::Button
-		);
-	MenuBuilder.AddMenuEntry(
-			LOCTEXT("HideAll", "Hide all"),
-			FText::GetEmpty(),
-			FSlateIcon(),
-			FUIAction(
-				FExecuteAction::CreateLambda([this]()
-				{
-					ForEachPropertyColumn([this](const FProperty& Property, FName ColumnId)
+				FSlateIcon(),
+				FUIAction(
+					FExecuteAction::CreateSP(this, &SConcertTransportLog::OnFilterMenuChecked),
+					FCanExecuteAction::CreateLambda([this] { return HeaderRow->IsColumnVisible(GET_MEMBER_NAME_CHECKED(FConcertLog, Timestamp)); }),
+					FIsActionChecked::CreateLambda([this] { return UConcertTransportLogSettings::GetSettings()->TimestampTimeFormat == ETimeFormat::Relative; })),
+				NAME_None,
+				EUserInterfaceActionType::ToggleButton
+			);
+
+		MenuBuilder.AddSeparator();
+		MenuBuilder.AddMenuEntry(
+				LOCTEXT("SelectAll", "Show all"),
+				FText::GetEmpty(),
+				FSlateIcon(),
+				FUIAction(
+					FExecuteAction::CreateLambda([this]()
 					{
-						HeaderRow->SetShowGeneratedColumn(ColumnId, false);
-					});
-				}),
-				FCanExecuteAction::CreateLambda([] { return true; })),
-			NAME_None,
-			EUserInterfaceActionType::Button
-		);
-	MenuBuilder.AddMenuEntry(
-			LOCTEXT("RestoreDefaultColumnVisibility", "Restore columns visibility"),
-			FText::GetEmpty(),
-			FSlateIcon(),
-			FUIAction(
-				FExecuteAction::CreateSP(this, &SConcertTransportLog::RestoreDefaultColumnVisiblities),
-				FCanExecuteAction::CreateLambda([] { return true; })),
-			NAME_None,
-			EUserInterfaceActionType::Button
-		);
-	MenuBuilder.AddSeparator();
-	UE::ConcertSharedSlate::AddEntriesForShowingHiddenRows(HeaderRow.ToSharedRef(), MenuBuilder);
-}
+						ForEachPropertyColumn([this](const FProperty& Property, FName ColumnId)
+						{
+							HeaderRow->SetShowGeneratedColumn(ColumnId, true);
+						});
+					}),
+					FCanExecuteAction::CreateLambda([] { return true; })),
+				NAME_None,
+				EUserInterfaceActionType::Button
+			);
+		MenuBuilder.AddMenuEntry(
+				LOCTEXT("HideAll", "Hide all"),
+				FText::GetEmpty(),
+				FSlateIcon(),
+				FUIAction(
+					FExecuteAction::CreateLambda([this]()
+					{
+						ForEachPropertyColumn([this](const FProperty& Property, FName ColumnId)
+						{
+							HeaderRow->SetShowGeneratedColumn(ColumnId, false);
+						});
+					}),
+					FCanExecuteAction::CreateLambda([] { return true; })),
+				NAME_None,
+				EUserInterfaceActionType::Button
+			);
+		MenuBuilder.AddMenuEntry(
+				LOCTEXT("RestoreDefaultColumnVisibility", "Restore columns visibility"),
+				FText::GetEmpty(),
+				FSlateIcon(),
+				FUIAction(
+					FExecuteAction::CreateSP(this, &SConcertTransportLog::RestoreDefaultColumnVisiblities),
+					FCanExecuteAction::CreateLambda([] { return true; })),
+				NAME_None,
+				EUserInterfaceActionType::Button
+			);
+		MenuBuilder.AddSeparator();
+		UE::ConcertSharedSlate::AddEntriesForShowingHiddenRows(HeaderRow.ToSharedRef(), MenuBuilder);
+	}
 
-void SConcertTransportLog::OnFilterMenuChecked()
-{
-	UConcertTransportLogSettings* Settings = UConcertTransportLogSettings::GetSettings();
-	
-	switch (Settings->TimestampTimeFormat)
+	void SConcertTransportLog::OnFilterMenuChecked()
 	{
-	case ETimeFormat::Relative:
-		Settings->TimestampTimeFormat = ETimeFormat::Absolute;
-		break;
-	case ETimeFormat::Absolute: 
-		Settings->TimestampTimeFormat = ETimeFormat::Relative;
-		break;
-	default:
-		checkNoEntry();
+		UConcertTransportLogSettings* Settings = UConcertTransportLogSettings::GetSettings();
+		
+		switch (Settings->TimestampTimeFormat)
+		{
+		case ETimeFormat::Relative:
+			Settings->TimestampTimeFormat = ETimeFormat::Absolute;
+			break;
+		case ETimeFormat::Absolute: 
+			Settings->TimestampTimeFormat = ETimeFormat::Relative;
+			break;
+		default:
+			checkNoEntry();
+		}
+
+		Settings->SaveConfig();
 	}
 
-	Settings->SaveConfig();
-}
+	void SConcertTransportLog::OnPageViewChanged(const TArray<TSharedPtr<FConcertLogEntry>>&)
+	{
+		LogView->RequestListRefresh();
 
-void SConcertTransportLog::OnPageViewChanged(const TArray<TSharedPtr<FConcertLogEntry>>&)
-{
-	LogView->RequestListRefresh();
+		if (bAutoScroll && PagedLogList->GetCurrentPage() == PagedLogList->GetNumPages() - 1)
+		{
+			LogView->ScrollToBottom();
+		}
+	}
 
-	if (bAutoScroll && PagedLogList->GetCurrentPage() == PagedLogList->GetNumPages() - 1)
+	void SConcertTransportLog::OnColumnVisibilitySettingsChanged(const FColumnVisibilitySnapshot& ColumnSnapshot)
 	{
-		LogView->ScrollToBottom();
+		TGuardValue<bool> GuardValue(bIsUpdatingColumnVisibility, true);
+		UE::ConcertSharedSlate::RestoreColumnVisibilityState(HeaderRow.ToSharedRef(), ColumnSnapshot);
 	}
-}
-
-void SConcertTransportLog::OnColumnVisibilitySettingsChanged(const FColumnVisibilitySnapshot& ColumnSnapshot)
-{
-	TGuardValue<bool> GuardValue(bIsUpdatingColumnVisibility, true);
-	UE::ConcertSharedSlate::RestoreColumnVisibilityState(HeaderRow.ToSharedRef(), ColumnSnapshot);
-}
 
-void SConcertTransportLog::OnConcertLoggingEnabledChanged(bool bNewEnabled)
-{
-	if (!bNewEnabled)
+	void SConcertTransportLog::OnConcertLoggingEnabledChanged(bool bNewEnabled)
 	{
-		EnableLoggingPromptOverlay->AddSlot().AttachWidget(SAssignNew(EnableLoggingPrompt, SPromptConcertLoggingEnabled));
+		if (!bNewEnabled)
+		{
+			EnableLoggingPromptOverlay->AddSlot().AttachWidget(SAssignNew(EnableLoggingPrompt, SPromptConcertLoggingEnabled));
+		}
+		else if (EnableLoggingPrompt)
+		{
+			EnableLoggingPromptOverlay->RemoveSlot(EnableLoggingPrompt.ToSharedRef());
+		}
 	}
-	else if (EnableLoggingPrompt)
+
+	bool SConcertTransportLog::CanScrollToAckLog(const FGuid& MessageId) const
 	{
-		EnableLoggingPromptOverlay->RemoveSlot(EnableLoggingPrompt.ToSharedRef());
+		return CanScrollToLog(MessageId, [MessageId](const FConcertLogEntry& Entry)
+		{
+			return SharedCanScrollToAckLog(MessageId, Entry);
+		});
 	}
-}
 
-bool SConcertTransportLog::CanScrollToAckLog(const FGuid& MessageId) const
-{
-	return CanScrollToLog(MessageId, [MessageId](const FConcertLogEntry& Entry)
+	bool SConcertTransportLog::CanScrollToAckedLog(const FGuid& MessageId) const
 	{
-		return SharedCanScrollToAckLog(MessageId, Entry);
-	});
-}
+		return CanScrollToLog(MessageId, [MessageId](const FConcertLogEntry& Entry)
+		{
+			return SharedCanScrollToAckedLog(MessageId, Entry);
+		});
+	}
 
-bool SConcertTransportLog::CanScrollToAckedLog(const FGuid& MessageId) const
-{
-	return CanScrollToLog(MessageId, [MessageId](const FConcertLogEntry& Entry)
+	void SConcertTransportLog::ScrollToAckLog(const FGuid& MessageId) const
 	{
-		return SharedCanScrollToAckedLog(MessageId, Entry);
-	});
-}
+		ScrollToLog(MessageId, [MessageId](const FConcertLogEntry& Entry)
+		{
+			return SharedCanScrollToAckLog(MessageId, Entry);
+		});
+	}
 
-void SConcertTransportLog::ScrollToAckLog(const FGuid& MessageId) const
-{
-	ScrollToLog(MessageId, [MessageId](const FConcertLogEntry& Entry)
+	void SConcertTransportLog::ScrollToAckedLog(const FGuid& MessageId) const
 	{
-		return SharedCanScrollToAckLog(MessageId, Entry);
-	});
-}
+		ScrollToLog(MessageId, [MessageId](const FConcertLogEntry& Entry)
+		{
+			return SharedCanScrollToAckedLog(MessageId, Entry);
+		});
+	}
 
-void SConcertTransportLog::ScrollToAckedLog(const FGuid& MessageId) const
-{
-	ScrollToLog(MessageId, [MessageId](const FConcertLogEntry& Entry)
+	bool SConcertTransportLog::SharedCanScrollToAckLog(const FGuid& MessageId, const FConcertLogEntry& Entry)
 	{
-		return SharedCanScrollToAckedLog(MessageId, Entry);
-	});
-}
-
-bool SConcertTransportLog::SharedCanScrollToAckLog(const FGuid& MessageId, const FConcertLogEntry& Entry)
-{
-	// MessageId is shared by some logs: there is a log for receiving one one for processing for example.
-	return Entry.Log.MessageId == MessageId && Entry.Log.MessageAction == EConcertLogMessageAction::Process;
-}
+		// MessageId is shared by some logs: there is a log for receiving one one for processing for example.
+		return Entry.Log.MessageId == MessageId && Entry.Log.MessageAction == EConcertLogMessageAction::Process;
+	}
 
-bool SConcertTransportLog::SharedCanScrollToAckedLog(const FGuid& MessageId, const FConcertLogEntry& Entry)
-{
-	// MessageId is shared by some logs: there is a log for receiving one one for sending for example.
-	return Entry.Log.MessageId == MessageId && Entry.Log.MessageAction == EConcertLogMessageAction::Send;
-}
+	bool SConcertTransportLog::SharedCanScrollToAckedLog(const FGuid& MessageId, const FConcertLogEntry& Entry)
+	{
+		// MessageId is shared by some logs: there is a log for receiving one one for sending for example.
+		return Entry.Log.MessageId == MessageId && Entry.Log.MessageAction == EConcertLogMessageAction::Send;
+	}
 
-void SConcertTransportLog::ScrollToLog(const int32 LogIndex) const
-{
-	if (const TOptional<FPagedFilteredConcertLogList::FPageCount> PageIndex = PagedLogList->GetPageOf(LogIndex); PageIndex && PagedLogList->GetFilteredLogs().IsValidIndex(LogIndex))
+	void SConcertTransportLog::ScrollToLog(const int32 LogIndex) const
 	{
-		PagedLogList->SetPage(*PageIndex);
-		LogView->RequestNavigateToItem(PagedLogList->GetFilteredLogs()[LogIndex]);
-		LogView->SetSelection(PagedLogList->GetFilteredLogs()[LogIndex]);
+		if (const TOptional<FPagedFilteredConcertLogList::FPageCount> PageIndex = PagedLogList->GetPageOf(LogIndex); PageIndex && PagedLogList->GetFilteredLogs().IsValidIndex(LogIndex))
+		{
+			PagedLogList->SetPage(*PageIndex);
+			LogView->RequestNavigateToItem(PagedLogList->GetFilteredLogs()[LogIndex]);
+			LogView->SetSelection(PagedLogList->GetFilteredLogs()[LogIndex]);
+		}
 	}
 }
 
diff --git a/Engine/Plugins/Developer/Concert/ConcertApp/MultiUserServer/Source/MultiUserServer/Private/Widgets/Clients/Logging/SConcertTransportLog.h b/Engine/Plugins/Developer/Concert/ConcertApp/MultiUserServer/Source/MultiUserServer/Private/Widgets/Clients/Logging/SConcertTransportLog.h
--- a/Engine/Plugins/Developer/Concert/ConcertApp/MultiUserServer/Source/MultiUserServer/Private/Widgets/Clients/Logging/SConcertTransportLog.h
+++ b/Engine/Plugins/Developer/Concert/ConcertApp/MultiUserServer/Source/MultiUserServer/Private/Widgets/Clients/Logging/SConcertTransportLog.h
@@ -7,8 +7,6 @@
 #include "Widgets/DeclarativeSyntaxSupport.h"
 #include "Widgets/SCompoundWidget.h"
 
-class FConcertLogTokenizer;
-class FConcertLogFilter_FrontendRoot;
 class FEndpointToUserNameCache;
 class FMenuBuilder;
 class FPagedFilteredConcertLogList;
@@ -23,74 +21,81 @@ template <typename ItemType> class SListView;
 struct FColumnVisibilitySnapshot;
 struct FConcertLogEntry;
 
-/**
- * Displays the contents of a IConcertLogSource and has UI for filtering.
- */
-class SConcertTransportLog : public SCompoundWidget
+namespace UE::MultiUserServer
 {
-public:
-
-	static const FName FirstColumnId;
-
-	SLATE_BEGIN_ARGS(SConcertTransportLog)
-	{}
-		/** Optional filters to display in UI */
-		SLATE_ARGUMENT(TSharedPtr<FConcertLogFilter_FrontendRoot>, Filter)
-	SLATE_END_ARGS()
-	virtual ~SConcertTransportLog() override;
-
-	void Construct(const FArguments& InArgs, TSharedRef<IConcertLogSource> LogSource, TSharedRef<FEndpointToUserNameCache> InEndpointCache, TSharedRef<FConcertLogTokenizer> InLogTokenizer);
-
-	bool CanScrollToLog(const FGuid& MessageId, FConcertLogEntryFilterFunc FilterFunc) const;
-	void ScrollToLog(const FGuid& MessageId, FConcertLogEntryFilterFunc FilterFunc) const;
-	
-private:
+	class FConcertLogFilter_FrontendRoot;
+	class FConcertLogTokenizer;
 	
-	/** Used to overlay EnableLoggingPrompt over the tabs */
-	TSharedPtr<SOverlay> EnableLoggingPromptOverlay;
-	/** Reminds the user to enable logging */
-	TSharedPtr<SPromptConcertLoggingEnabled> EnableLoggingPrompt;
+	/**
+	 * Displays the contents of a IConcertLogSource and has UI for filtering.
+	 */
+	class SConcertTransportLog : public SCompoundWidget
+	{
+	public:
+
+		static const FName FirstColumnId;
+
+		SLATE_BEGIN_ARGS(SConcertTransportLog)
+		{}
+			/** Optional filters to display in UI */
+			SLATE_ARGUMENT(TSharedPtr<FConcertLogFilter_FrontendRoot>, Filter)
+		SLATE_END_ARGS()
+		virtual ~SConcertTransportLog() override;
+
+		void Construct(const FArguments& InArgs, TSharedRef<IConcertLogSource> LogSource, TSharedRef<FEndpointToUserNameCache> InEndpointCache, TSharedRef<FConcertLogTokenizer> InLogTokenizer);
+
+		bool CanScrollToLog(const FGuid& MessageId, FConcertLogEntryFilterFunc FilterFunc) const;
+		void ScrollToLog(const FGuid& MessageId, FConcertLogEntryFilterFunc FilterFunc) const;
+		
+	private:
+		
+		/** Used to overlay EnableLoggingPrompt over the tabs */
+		TSharedPtr<SOverlay> EnableLoggingPromptOverlay;
+		/** Reminds the user to enable logging */
+		TSharedPtr<SPromptConcertLoggingEnabled> EnableLoggingPrompt;
+
+		TSharedPtr<FEndpointToUserNameCache> EndpointCache;
+		
+		/** Sorts the log into pages whilst applying filters */
+		TSharedPtr<FPagedFilteredConcertLogList> PagedLogList;
+		/** Used by various systems to convert logs to text */
+		TSharedPtr<FConcertLogTokenizer> LogTokenizer;
+
+		TSharedPtr<FConcertLogFilter_FrontendRoot> Filter;
+		
+		/** Lists the logs */
+		TSharedPtr<SListView<TSharedPtr<FConcertLogEntry>>> LogView;
+		/** Header row of LogView */
+		TSharedPtr<SHeaderRow> HeaderRow;
+
+		/** Whether to automatically scroll to new logs as they come in */
+		bool bAutoScroll = true;
+		/** Whether we currently loading the column visibility - prevents infinite event recursion */
+		bool bIsUpdatingColumnVisibility = false;
+
+		// Table view creation
+		TSharedRef<SWidget> CreateTableView();
+		TSharedRef<SHeaderRow> CreateHeaderRow();
+		TSharedRef<ITableRow> OnGenerateActivityRowWidget(TSharedPtr<FConcertLogEntry> Item, const TSharedRef<STableViewBase>& OwnerTable) const;
+
+		static void ForEachPropertyColumn(TFunctionRef<void(const FProperty& ColumnPropertyId, FName ColumnId)> Callback);
+
+		void RestoreDefaultColumnVisiblities();
+		void ExtendViewOptions(FMenuBuilder& MenuBuilder);
+		void OnFilterMenuChecked();
+		void OnPageViewChanged(const TArray<TSharedPtr<FConcertLogEntry>>&);
+		void OnColumnVisibilitySettingsChanged(const FColumnVisibilitySnapshot& ColumnSnapshot);
+		
+		void OnConcertLoggingEnabledChanged(bool bNewEnabled);
+
+		bool CanScrollToAckLog(const FGuid& MessageId) const;
+		bool CanScrollToAckedLog(const FGuid& MessageId) const;
+		void ScrollToAckLog(const FGuid& MessageId) const;
+		void ScrollToAckedLog(const FGuid& MessageId) const;
+		static bool SharedCanScrollToAckLog(const FGuid& MessageId, const FConcertLogEntry& Entry);
+		static bool SharedCanScrollToAckedLog(const FGuid& MessageId, const FConcertLogEntry& Entry);
+		
+		void ScrollToLog(const int32 LogIndex) const;
+	};
+}
 
-	TSharedPtr<FEndpointToUserNameCache> EndpointCache;
-	
-	/** Sorts the log into pages whilst applying filters */
-	TSharedPtr<FPagedFilteredConcertLogList> PagedLogList;
-	/** Used by various systems to convert logs to text */
-	TSharedPtr<FConcertLogTokenizer> LogTokenizer;
-
-	TSharedPtr<FConcertLogFilter_FrontendRoot> Filter;
-	
-	/** Lists the logs */
-	TSharedPtr<SListView<TSharedPtr<FConcertLogEntry>>> LogView;
-	/** Header row of LogView */
-	TSharedPtr<SHeaderRow> HeaderRow;
-
-	/** Whether to automatically scroll to new logs as they come in */
-	bool bAutoScroll = true;
-	/** Whether we currently loading the column visibility - prevents infinite event recursion */
-	bool bIsUpdatingColumnVisibility = false;
-
-	// Table view creation
-	TSharedRef<SWidget> CreateTableView();
-	TSharedRef<SHeaderRow> CreateHeaderRow();
-	TSharedRef<ITableRow> OnGenerateActivityRowWidget(TSharedPtr<FConcertLogEntry> Item, const TSharedRef<STableViewBase>& OwnerTable) const;
-
-	static void ForEachPropertyColumn(TFunctionRef<void(const FProperty& ColumnPropertyId, FName ColumnId)> Callback);
-
-	void RestoreDefaultColumnVisiblities();
-	void ExtendViewOptions(FMenuBuilder& MenuBuilder);
-	void OnFilterMenuChecked();
-	void OnPageViewChanged(const TArray<TSharedPtr<FConcertLogEntry>>&);
-	void OnColumnVisibilitySettingsChanged(const FColumnVisibilitySnapshot& ColumnSnapshot);
-	
-	void OnConcertLoggingEnabledChanged(bool bNewEnabled);
-
-	bool CanScrollToAckLog(const FGuid& MessageId) const;
-	bool CanScrollToAckedLog(const FGuid& MessageId) const;
-	void ScrollToAckLog(const FGuid& MessageId) const;
-	void ScrollToAckedLog(const FGuid& MessageId) const;
-	static bool SharedCanScrollToAckLog(const FGuid& MessageId, const FConcertLogEntry& Entry);
-	static bool SharedCanScrollToAckedLog(const FGuid& MessageId, const FConcertLogEntry& Entry);
-	
-	void ScrollToLog(const int32 LogIndex) const;
-};
diff --git a/Engine/Plugins/Developer/Concert/ConcertApp/MultiUserServer/Source/MultiUserServer/Private/Widgets/Clients/Logging/SConcertTransportLogFooter.cpp b/Engine/Plugins/Developer/Concert/ConcertApp/MultiUserServer/Source/MultiUserServer/Private/Widgets/Clients/Logging/SConcertTransportLogFooter.cpp
--- a/Engine/Plugins/Developer/Concert/ConcertApp/MultiUserServer/Source/MultiUserServer/Private/Widgets/Clients/Logging/SConcertTransportLogFooter.cpp
+++ b/Engine/Plugins/Developer/Concert/ConcertApp/MultiUserServer/Source/MultiUserServer/Private/Widgets/Clients/Logging/SConcertTransportLogFooter.cpp
@@ -14,134 +14,138 @@
 
 #define LOCTEXT_NAMESPACE "UnrealMultiUserUI"
 
-namespace UE::MultiUserServer::Private
-{
-	constexpr size_t DefaultPageSizeIndex = 3; 
-	static TArray<TSharedPtr<FPagedFilteredConcertLogList::FLogsPerPageCount>> PageSizeOptions = {
-		MakeShared<FPagedFilteredConcertLogList::FLogsPerPageCount>(10),
-		MakeShared<FPagedFilteredConcertLogList::FLogsPerPageCount>(50),
-		MakeShared<FPagedFilteredConcertLogList::FLogsPerPageCount>(100),
-		MakeShared<FPagedFilteredConcertLogList::FLogsPerPageCount>(500),
-		MakeShared<FPagedFilteredConcertLogList::FLogsPerPageCount>(1000),
-		MakeShared<FPagedFilteredConcertLogList::FLogsPerPageCount>(5000)
-	};
-}
 
-void SConcertTransportLogFooter::Construct(const FArguments& InArgs, TSharedRef<FPagedFilteredConcertLogList> InPagedLogList)
+namespace UE::MultiUserServer
 {
-	using namespace UE::MultiUserServer::Private;
-	PagedLogList = MoveTemp(InPagedLogList);
-
-	ChildSlot
-	[
-		SNew(SHorizontalBox)
-		
-		// Page selection
-		+SHorizontalBox::Slot()
-		.AutoWidth()
-		.VAlign(VAlign_Center)
-		[
-			SNew(SSpinBox<FPagedFilteredConcertLogList::FPageCount>)
-			.Value_Lambda([this](){ return PagedLogList->GetFilteredLogs().Num() == 0 ? 0 : PagedLogList->GetCurrentPage() + 1; })
-			.OnValueChanged_Lambda([this](FPagedFilteredConcertLogList::FPageCount NewValue) { PagedLogList->SetPage(NewValue - 1); })
-			.MinValue(0)
-			.MaxValue_Lambda([this]() { return PagedLogList->GetNumPages(); })
-		]
-
-		// Page counter
-		+SHorizontalBox::Slot()
-		.AutoWidth()
-		.Padding(2, 0, 0, 0)
-		.VAlign(VAlign_Center)
-		[
-			SNew(STextBlock)
-			.Text_Lambda([this](){ return FText::Format(LOCTEXT("PageCounterFormat", "of {0} {0}|plural(one=page,other=pages)"), PagedLogList->GetNumPages()); })
-		]
-
-		// Logs per page text
-		+SHorizontalBox::Slot()
-		.AutoWidth()
-		.Padding(30, 0, 0, 0)
-		.VAlign(VAlign_Center)
+	namespace Private
+	{
+		constexpr size_t DefaultPageSizeIndex = 3; 
+		static TArray<TSharedPtr<FPagedFilteredConcertLogList::FLogsPerPageCount>> PageSizeOptions = {
+			MakeShared<FPagedFilteredConcertLogList::FLogsPerPageCount>(10),
+			MakeShared<FPagedFilteredConcertLogList::FLogsPerPageCount>(50),
+			MakeShared<FPagedFilteredConcertLogList::FLogsPerPageCount>(100),
+			MakeShared<FPagedFilteredConcertLogList::FLogsPerPageCount>(500),
+			MakeShared<FPagedFilteredConcertLogList::FLogsPerPageCount>(1000),
+			MakeShared<FPagedFilteredConcertLogList::FLogsPerPageCount>(5000)
+		};
+	}
+
+	void SConcertTransportLogFooter::Construct(const FArguments& InArgs, TSharedRef<FPagedFilteredConcertLogList> InPagedLogList)
+	{
+		using namespace UE::MultiUserServer::Private;
+		PagedLogList = MoveTemp(InPagedLogList);
+
+		ChildSlot
 		[
-			SNew(STextBlock)
-			.Text(LOCTEXT("LogePerPageText", "Logs per page"))
-		]
-		
-		// Logs per page combo box
-		+SHorizontalBox::Slot()
-		.AutoWidth()
-		.Padding(3, 0, 0, 0)
-		.VAlign(VAlign_Center)
-		[
-			SNew(SComboBox<TSharedPtr<FPagedFilteredConcertLogList::FLogsPerPageCount>>)
-			.InitiallySelectedItem(PageSizeOptions[DefaultPageSizeIndex])
-			.OnSelectionChanged_Lambda([this](const TSharedPtr<FPagedFilteredConcertLogList::FLogsPerPageCount>& NewCount, ESelectInfo::Type){ PagedLogList->SetLogsPerPage(*NewCount); })
-			.OnGenerateWidget_Lambda([](const TSharedPtr<FPagedFilteredConcertLogList::FLogsPerPageCount>& Count) { return SNew(STextBlock).Text(FText::FromString(FString::FromInt(*Count))); })
-			.OptionsSource(&PageSizeOptions)
+			SNew(SHorizontalBox)
+			
+			// Page selection
+			+SHorizontalBox::Slot()
+			.AutoWidth()
+			.VAlign(VAlign_Center)
 			[
-				SNew(STextBlock).Text_Lambda([this](){ return FText::FromString(FString::FromInt(PagedLogList->GetLogsPerPage())); })
+				SNew(SSpinBox<FPagedFilteredConcertLogList::FPageCount>)
+				.Value_Lambda([this](){ return PagedLogList->GetFilteredLogs().Num() == 0 ? 0 : PagedLogList->GetCurrentPage() + 1; })
+				.OnValueChanged_Lambda([this](FPagedFilteredConcertLogList::FPageCount NewValue) { PagedLogList->SetPage(NewValue - 1); })
+				.MinValue(0)
+				.MaxValue_Lambda([this]() { return PagedLogList->GetNumPages(); })
 			]
-		]
 
-		// Log counter
-		+SHorizontalBox::Slot()
-		.AutoWidth()
-		.Padding(30, 0, 0, 0)
-		.VAlign(VAlign_Center)
-		[
-			SNew(STextBlock)
-			.Text_Lambda([this]()
-			{
-				return PagedLogList->IsBufferFull()
-					? FText::Format(LOCTEXT("LogCounterFormat.Full", "{0} filtered {0}|plural(one=log,other=logs) (buffer full)"), PagedLogList->GetFilteredLogs().Num())
-					: FText::Format(LOCTEXT("LogCounterFormat.NotFull", "{0} filtered {0}|plural(one=log,other=logs)"), PagedLogList->GetFilteredLogs().Num());
-			})
-		]
-
-		// Gap Filler
-		+SHorizontalBox::Slot()
-		.FillWidth(1.0)
-		[
-			SNew(SSpacer)
-		]
+			// Page counter
+			+SHorizontalBox::Slot()
+			.AutoWidth()
+			.Padding(2, 0, 0, 0)
+			.VAlign(VAlign_Center)
+			[
+				SNew(STextBlock)
+				.Text_Lambda([this](){ return FText::Format(LOCTEXT("PageCounterFormat", "of {0} {0}|plural(one=page,other=pages)"), PagedLogList->GetNumPages()); })
+			]
 
-		+SHorizontalBox::Slot()
-		.AutoWidth()
-		.VAlign(VAlign_Center)
-		[
-			SNew(SComboButton)
-				.OnGetMenuContent(this, &SConcertTransportLogFooter::MakeViewOptionsMenuWidget, InArgs._ExtendViewOptions)
-				.HasDownArrow(true)
-				.ButtonContent()
+			// Logs per page text
+			+SHorizontalBox::Slot()
+			.AutoWidth()
+			.Padding(30, 0, 0, 0)
+			.VAlign(VAlign_Center)
+			[
+				SNew(STextBlock)
+				.Text(LOCTEXT("LogePerPageText", "Logs per page"))
+			]
+			
+			// Logs per page combo box
+			+SHorizontalBox::Slot()
+			.AutoWidth()
+			.Padding(3, 0, 0, 0)
+			.VAlign(VAlign_Center)
+			[
+				SNew(SComboBox<TSharedPtr<FPagedFilteredConcertLogList::FLogsPerPageCount>>)
+				.InitiallySelectedItem(PageSizeOptions[DefaultPageSizeIndex])
+				.OnSelectionChanged_Lambda([this](const TSharedPtr<FPagedFilteredConcertLogList::FLogsPerPageCount>& NewCount, ESelectInfo::Type){ PagedLogList->SetLogsPerPage(*NewCount); })
+				.OnGenerateWidget_Lambda([](const TSharedPtr<FPagedFilteredConcertLogList::FLogsPerPageCount>& Count) { return SNew(STextBlock).Text(FText::FromString(FString::FromInt(*Count))); })
+				.OptionsSource(&PageSizeOptions)
 				[
-					SNew(SHorizontalBox)
+					SNew(STextBlock).Text_Lambda([this](){ return FText::FromString(FString::FromInt(PagedLogList->GetLogsPerPage())); })
+				]
+			]
 
-					+SHorizontalBox::Slot()
-					.AutoWidth()
-					.VAlign(VAlign_Center)
-					[
-						SNew(SImage).Image(FAppStyle::Get().GetBrush("Icons.Visible")) // The eye ball image.
-					]
+			// Log counter
+			+SHorizontalBox::Slot()
+			.AutoWidth()
+			.Padding(30, 0, 0, 0)
+			.VAlign(VAlign_Center)
+			[
+				SNew(STextBlock)
+				.Text_Lambda([this]()
+				{
+					return PagedLogList->IsBufferFull()
+						? FText::Format(LOCTEXT("LogCounterFormat.Full", "{0} filtered {0}|plural(one=log,other=logs) (buffer full)"), PagedLogList->GetFilteredLogs().Num())
+						: FText::Format(LOCTEXT("LogCounterFormat.NotFull", "{0} filtered {0}|plural(one=log,other=logs)"), PagedLogList->GetFilteredLogs().Num());
+				})
+			]
 
-					+SHorizontalBox::Slot()
-					.AutoWidth()
-					.Padding(2, 0, 0, 0)
-					.VAlign(VAlign_Center)
+			// Gap Filler
+			+SHorizontalBox::Slot()
+			.FillWidth(1.0)
+			[
+				SNew(SSpacer)
+			]
+
+			+SHorizontalBox::Slot()
+			.AutoWidth()
+			.VAlign(VAlign_Center)
+			[
+				SNew(SComboButton)
+					.OnGetMenuContent(this, &SConcertTransportLogFooter::MakeViewOptionsMenuWidget, InArgs._ExtendViewOptions)
+					.HasDownArrow(true)
+					.ButtonContent()
 					[
-						SNew(STextBlock).Text(LOCTEXT("ViewOptions", "View Options"))
+						SNew(SHorizontalBox)
+
+						+SHorizontalBox::Slot()
+						.AutoWidth()
+						.VAlign(VAlign_Center)
+						[
+							SNew(SImage).Image(FAppStyle::Get().GetBrush("Icons.Visible")) // The eye ball image.
+						]
+
+						+SHorizontalBox::Slot()
+						.AutoWidth()
+						.Padding(2, 0, 0, 0)
+						.VAlign(VAlign_Center)
+						[
+							SNew(STextBlock).Text(LOCTEXT("ViewOptions", "View Options"))
+						]
 					]
-				]
-		]
-	];
-}
+			]
+		];
+	}
 
-TSharedRef<SWidget> SConcertTransportLogFooter::MakeViewOptionsMenuWidget(FExtendContextMenu ExtendViewOptions)
-{
-	FMenuBuilder MenuBuilder(/*bInShouldCloseWindowAfterMenuSelection=*/true, nullptr);
+	TSharedRef<SWidget> SConcertTransportLogFooter::MakeViewOptionsMenuWidget(FExtendContextMenu ExtendViewOptions)
+	{
+		FMenuBuilder MenuBuilder(/*bInShouldCloseWindowAfterMenuSelection=*/true, nullptr);
 
-	ExtendViewOptions.Execute(MenuBuilder);
-	return MenuBuilder.MakeWidget();
+		ExtendViewOptions.Execute(MenuBuilder);
+		return MenuBuilder.MakeWidget();
+	}
 }
 
 #undef LOCTEXT_NAMESPACE
diff --git a/Engine/Plugins/Developer/Concert/ConcertApp/MultiUserServer/Source/MultiUserServer/Private/Widgets/Clients/Logging/SConcertTransportLogFooter.h b/Engine/Plugins/Developer/Concert/ConcertApp/MultiUserServer/Source/MultiUserServer/Private/Widgets/Clients/Logging/SConcertTransportLogFooter.h
--- a/Engine/Plugins/Developer/Concert/ConcertApp/MultiUserServer/Source/MultiUserServer/Private/Widgets/Clients/Logging/SConcertTransportLogFooter.h
+++ b/Engine/Plugins/Developer/Concert/ConcertApp/MultiUserServer/Source/MultiUserServer/Private/Widgets/Clients/Logging/SConcertTransportLogFooter.h
@@ -12,28 +12,31 @@ class STextBlock;
 template<typename T> class SSpinBox;
 struct FConcertLogEntry;
 
-/** Displays the number of pages and items. Displayed under the table view. */
-class SConcertTransportLogFooter : public SCompoundWidget
+namespace UE::MultiUserServer
 {
-public:
+	/** Displays the number of pages and items. Displayed under the table view. */
+	class SConcertTransportLogFooter : public SCompoundWidget
+	{
+	public:
 	
-	DECLARE_DELEGATE_OneParam(FExtendContextMenu, FMenuBuilder&)
+		DECLARE_DELEGATE_OneParam(FExtendContextMenu, FMenuBuilder&)
 
-	SLATE_BEGIN_ARGS(SConcertTransportLogFooter) {}
+		SLATE_BEGIN_ARGS(SConcertTransportLogFooter) {}
 		SLATE_EVENT(FExtendContextMenu, ExtendViewOptions)
 	SLATE_END_ARGS()
 
 	void Construct(const FArguments& InArgs, TSharedRef<FPagedFilteredConcertLogList> InPagedLogList);
 
-private:
+	private:
 
-	/** The model we'll be updating */
-	TSharedPtr<FPagedFilteredConcertLogList> PagedLogList;
+		/** The model we'll be updating */
+		TSharedPtr<FPagedFilteredConcertLogList> PagedLogList;
 
-	/** Selects the current page */
-	TSharedPtr<SSpinBox<FPagedFilteredConcertLogList::FPageCount>> CurrentPage;
-	/** Displays the number of pages */
-	TSharedPtr<STextBlock> PageCounterText;
+		/** Selects the current page */
+		TSharedPtr<SSpinBox<FPagedFilteredConcertLogList::FPageCount>> CurrentPage;
+		/** Displays the number of pages */
+		TSharedPtr<STextBlock> PageCounterText;
 
-	TSharedRef<SWidget> MakeViewOptionsMenuWidget(FExtendContextMenu ExtendViewOptions);
-};
+		TSharedRef<SWidget> MakeViewOptionsMenuWidget(FExtendContextMenu ExtendViewOptions);
+	};
+}
diff --git a/Engine/Plugins/Developer/Concert/ConcertApp/MultiUserServer/Source/MultiUserServer/Private/Widgets/Clients/Logging/SConcertTransportLogRow.cpp b/Engine/Plugins/Developer/Concert/ConcertApp/MultiUserServer/Source/MultiUserServer/Private/Widgets/Clients/Logging/SConcertTransportLogRow.cpp
--- a/Engine/Plugins/Developer/Concert/ConcertApp/MultiUserServer/Source/MultiUserServer/Private/Widgets/Clients/Logging/SConcertTransportLogRow.cpp
+++ b/Engine/Plugins/Developer/Concert/ConcertApp/MultiUserServer/Source/MultiUserServer/Private/Widgets/Clients/Logging/SConcertTransportLogRow.cpp
@@ -3,12 +3,12 @@
 #include "SConcertTransportLogRow.h"
 
 #include "ConcertServerStyle.h"
-#include "Util/ConcertLogTokenizer.h"
 #include "SConcertTransportLog.h"
-#include "Algo/AllOf.h"
-#include "Algo/AnyOf.h"
 #include "Session/Activity/PredefinedActivityColumns.h"
+#include "Util/ConcertLogTokenizer.h"
 
+#include "Algo/AllOf.h"
+#include "Algo/AnyOf.h"
 #include "Framework/MultiBox/MultiBoxBuilder.h"
 #include "Widgets/Colors/SColorBlock.h"
 #include "Widgets/Input/SButton.h"
@@ -20,250 +20,253 @@
 
 #define LOCTEXT_NAMESPACE "UnrealMultiUserUI.SConcertTransportLogRow"
 
-void SConcertTransportLogRow::Construct(const FArguments& InArgs, TSharedPtr<FConcertLogEntry> InLogEntry, const TSharedRef<STableViewBase>& InOwnerTableView, TSharedRef<FConcertLogTokenizer> InTokenizer)
-{
-	LogEntry = MoveTemp(InLogEntry);
-	Tokenizer = MoveTemp(InTokenizer);
-	HighlightText = InArgs._HighlightText;
-	
-	AvatarColor = InArgs._AvatarColor;
-	CanScrollToAckLogFunc = InArgs._CanScrollToAckLog;
-	CanScrollToAckedLogFunc = InArgs._CanScrollToAckedLog;
-	ScrollToAckLogFunc = InArgs._ScrollToAckLog;
-	ScrollToAckedLogFunc = InArgs._ScrollToAckedLog;
-	check(ScrollToAckLogFunc.IsBound() && ScrollToAckedLogFunc.IsBound());
-	
-	SMultiColumnTableRow<TSharedPtr<FConcertLogEntry>>::Construct(FSuperRowType::FArguments(), InOwnerTableView);
-}
-
-TSharedRef<SWidget> SConcertTransportLogRow::GenerateWidgetForColumn(const FName& ColumnName)
+namespace UE::MultiUserServer
 {
-	if (ColumnName == SConcertTransportLog::FirstColumnId)
+	void SConcertTransportLogRow::Construct(const FArguments& InArgs, TSharedPtr<FConcertLogEntry> InLogEntry, const TSharedRef<STableViewBase>& InOwnerTableView, TSharedRef<FConcertLogTokenizer> InTokenizer)
 	{
-		return SNew(SHorizontalBox)
-			+SHorizontalBox::Slot()
-			.Padding(2, 1)
-			[
-				SNew(SColorBlock)
-				.Color(AvatarColor)
-				.Size(FVector2D(4.f, 16.f))
-			];
+		LogEntry = MoveTemp(InLogEntry);
+		Tokenizer = MoveTemp(InTokenizer);
+		HighlightText = InArgs._HighlightText;
+		
+		AvatarColor = InArgs._AvatarColor;
+		CanScrollToAckLogFunc = InArgs._CanScrollToAckLog;
+		CanScrollToAckedLogFunc = InArgs._CanScrollToAckedLog;
+		ScrollToAckLogFunc = InArgs._ScrollToAckLog;
+		ScrollToAckedLogFunc = InArgs._ScrollToAckedLog;
+		check(ScrollToAckLogFunc.IsBound() && ScrollToAckedLogFunc.IsBound());
+		
+		SMultiColumnTableRow<TSharedPtr<FConcertLogEntry>>::Construct(FSuperRowType::FArguments(), InOwnerTableView);
 	}
-	
-	using FColumnWidgetFactoryFunc = TSharedRef<SWidget>(SConcertTransportLogRow::*)(const FName& InColumnName);
-	const TMap<FName, FColumnWidgetFactoryFunc> OverrideFactories = {
-		{ GET_MEMBER_NAME_CHECKED(FConcertLogMetadata, AckState), &SConcertTransportLogRow::CreateAckColumn }
-	};
-	const FColumnWidgetFactoryFunc FallbackFactoryFunc = &SConcertTransportLogRow::CreateDefaultColumn;
 
-	const FColumnWidgetFactoryFunc* FactoryFunc = OverrideFactories.Find(ColumnName);
-	return FactoryFunc
-		? Invoke(*FactoryFunc, this, ColumnName)
-		: Invoke(FallbackFactoryFunc, this, ColumnName);
-}
-
-TSharedRef<SWidget> SConcertTransportLogRow::CreateDefaultColumn(const FName& PropertyName)
-{
-	return SNew(STextBlock)
-		.Text_Lambda([PropertyName, WeakEntry = TWeakPtr<FConcertLogEntry>(LogEntry), WeakTokenizer = TWeakPtr<FConcertLogTokenizer>(Tokenizer)]()
+	TSharedRef<SWidget> SConcertTransportLogRow::GenerateWidgetForColumn(const FName& ColumnName)
+	{
+		if (ColumnName == SConcertTransportLog::FirstColumnId)
 		{
-			TSharedPtr<FConcertLogEntry> PinnedEntry = WeakEntry.Pin();
-			TSharedPtr<FConcertLogTokenizer> PinnedTokenizer = WeakTokenizer.Pin();
-			if (!ensure(PinnedEntry && PinnedTokenizer))
-			{
-				return FText::GetEmpty();
-			}
-			
-			const FProperty* ConcertLogProperty = FConcertLog::StaticStruct()->FindPropertyByName(PropertyName);
-			if (ConcertLogProperty)
-			{
-				return FText::FromString(PinnedTokenizer->Tokenize(PinnedEntry->Log, *ConcertLogProperty));
-			}
+			return SNew(SHorizontalBox)
+				+SHorizontalBox::Slot()
+				.Padding(2, 1)
+				[
+					SNew(SColorBlock)
+					.Color(AvatarColor)
+					.Size(FVector2D(4.f, 16.f))
+				];
+		}
+		
+		using FColumnWidgetFactoryFunc = TSharedRef<SWidget>(SConcertTransportLogRow::*)(const FName& InColumnName);
+		const TMap<FName, FColumnWidgetFactoryFunc> OverrideFactories = {
+			{ GET_MEMBER_NAME_CHECKED(FConcertLogMetadata, AckState), &SConcertTransportLogRow::CreateAckColumn }
+		};
+		const FColumnWidgetFactoryFunc FallbackFactoryFunc = &SConcertTransportLogRow::CreateDefaultColumn;
 
-			const FProperty* ConcertLogEntryProperty = FConcertLogMetadata::StaticStruct()->FindPropertyByName(PropertyName);
-			if (ConcertLogEntryProperty)
+		const FColumnWidgetFactoryFunc* FactoryFunc = OverrideFactories.Find(ColumnName);
+		return FactoryFunc
+			? Invoke(*FactoryFunc, this, ColumnName)
+			: Invoke(FallbackFactoryFunc, this, ColumnName);
+	}
+
+	TSharedRef<SWidget> SConcertTransportLogRow::CreateDefaultColumn(const FName& PropertyName)
+	{
+		return SNew(STextBlock)
+			.Text_Lambda([PropertyName, WeakEntry = TWeakPtr<FConcertLogEntry>(LogEntry), WeakTokenizer = TWeakPtr<FConcertLogTokenizer>(Tokenizer)]()
 			{
-				return FText::FromString(PinnedTokenizer->Tokenize(PinnedEntry->LogMetaData, *ConcertLogEntryProperty));
-			}
+				TSharedPtr<FConcertLogEntry> PinnedEntry = WeakEntry.Pin();
+				TSharedPtr<FConcertLogTokenizer> PinnedTokenizer = WeakTokenizer.Pin();
+				if (!ensure(PinnedEntry && PinnedTokenizer))
+				{
+					return FText::GetEmpty();
+				}
 				
-			return FText::GetEmpty();
-		})
-		.HighlightText(HighlightText);
-}
+				const FProperty* ConcertLogProperty = FConcertLog::StaticStruct()->FindPropertyByName(PropertyName);
+				if (ConcertLogProperty)
+				{
+					return FText::FromString(PinnedTokenizer->Tokenize(PinnedEntry->Log, *ConcertLogProperty));
+				}
 
-TSharedRef<SWidget> SConcertTransportLogRow::CreateAckColumn(const FName& PropertyName)
-{
-	return SNew(SOverlay)
+				const FProperty* ConcertLogEntryProperty = FConcertLogMetadata::StaticStruct()->FindPropertyByName(PropertyName);
+				if (ConcertLogEntryProperty)
+				{
+					return FText::FromString(PinnedTokenizer->Tokenize(PinnedEntry->LogMetaData, *ConcertLogEntryProperty));
+				}
+					
+				return FText::GetEmpty();
+			})
+			.HighlightText(HighlightText);
+	}
 
-		+SOverlay::Slot()
-		.HAlign(HAlign_Center)
-		.VAlign(VAlign_Center)
-		[
-			CreateAckInProgressWidget()
-		]
-	
-		+SOverlay::Slot()
-		.HAlign(HAlign_Center)
-		.VAlign(VAlign_Center)
-		[
-			CreateAckFailureWidget()
-		]
+	TSharedRef<SWidget> SConcertTransportLogRow::CreateAckColumn(const FName& PropertyName)
+	{
+		return SNew(SOverlay)
 
-		+SOverlay::Slot()
-		.HAlign(HAlign_Center)
-		.VAlign(VAlign_Center)
-		[
-			CreateAckNotNeededWidget()
-		]
+			+SOverlay::Slot()
+			.HAlign(HAlign_Center)
+			.VAlign(VAlign_Center)
+			[
+				CreateAckInProgressWidget()
+			]
+		
+			+SOverlay::Slot()
+			.HAlign(HAlign_Center)
+			.VAlign(VAlign_Center)
+			[
+				CreateAckFailureWidget()
+			]
 
-		+SOverlay::Slot()
-		.HAlign(HAlign_Center)
-		.VAlign(VAlign_Center)
-		[
-			CreateAckSuccessWidget()
-		]
+			+SOverlay::Slot()
+			.HAlign(HAlign_Center)
+			.VAlign(VAlign_Center)
+			[
+				CreateAckNotNeededWidget()
+			]
 
-		+SOverlay::Slot()
-		.HAlign(HAlign_Center)
-		.VAlign(VAlign_Center)
-		[
-			CreateAckWidget()
-		]
-	;
-}
+			+SOverlay::Slot()
+			.HAlign(HAlign_Center)
+			.VAlign(VAlign_Center)
+			[
+				CreateAckSuccessWidget()
+			]
 
-TSharedRef<SWidget> SConcertTransportLogRow::CreateAckInProgressWidget() const
-{
-	return SNew(SScaleBox)
-		.Stretch(EStretch::ScaleToFit)
-		[
-			SNew(SThrobber)
-			.Animate(SThrobber::VerticalAndOpacity)
+			+SOverlay::Slot()
+			.HAlign(HAlign_Center)
+			.VAlign(VAlign_Center)
+			[
+				CreateAckWidget()
+			]
+		;
+	}
+
+	TSharedRef<SWidget> SConcertTransportLogRow::CreateAckInProgressWidget() const
+	{
+		return SNew(SScaleBox)
+			.Stretch(EStretch::ScaleToFit)
+			[
+				SNew(SThrobber)
+				.Animate(SThrobber::VerticalAndOpacity)
+				.Visibility_Lambda([this]()
+				{
+					return LogEntry->LogMetaData.AckState == EConcertLogAckState::InProgress
+						? EVisibility::Visible
+						: EVisibility::Hidden;
+				})
+				.ToolTipText(LOCTEXT("InProgress", "Awaiting ACK..."))
+			];
+	}
+
+	TSharedRef<SWidget> SConcertTransportLogRow::CreateAckSuccessWidget() const
+	{
+		return SNew(SButton)
+			.ContentPadding(0.f)
+			.IsEnabled_Lambda([this]()
+			{
+				return LogEntry->LogMetaData.AckingMessageId.IsSet() && CanScrollToAckLogFunc.Execute(*LogEntry->LogMetaData.AckingMessageId);
+			})
+			.OnPressed_Lambda([this]()
+			{
+				if (LogEntry->LogMetaData.AckingMessageId.IsSet())
+				{
+					ScrollToAckLogFunc.Execute(*LogEntry->LogMetaData.AckingMessageId);
+				}
+			})
 			.Visibility_Lambda([this]()
 			{
-				return LogEntry->LogMetaData.AckState == EConcertLogAckState::InProgress
+				return LogEntry->LogMetaData.AckState == EConcertLogAckState::AckReceived
 					? EVisibility::Visible
 					: EVisibility::Hidden;
 			})
-			.ToolTipText(LOCTEXT("InProgress", "Awaiting ACK..."))
-		];
-}
-
-TSharedRef<SWidget> SConcertTransportLogRow::CreateAckSuccessWidget() const
-{
-	return SNew(SButton)
-		.ContentPadding(0.f)
-		.IsEnabled_Lambda([this]()
-		{
-			return LogEntry->LogMetaData.AckingMessageId.IsSet() && CanScrollToAckLogFunc.Execute(*LogEntry->LogMetaData.AckingMessageId);
-		})
-		.OnPressed_Lambda([this]()
-		{
-			if (LogEntry->LogMetaData.AckingMessageId.IsSet())
+			.ToolTipText_Lambda([this]()
 			{
-				ScrollToAckLogFunc.Execute(*LogEntry->LogMetaData.AckingMessageId);
-			}
-		})
-		.Visibility_Lambda([this]()
-		{
-			return LogEntry->LogMetaData.AckState == EConcertLogAckState::AckReceived
-				? EVisibility::Visible
-				: EVisibility::Hidden;
-		})
-		.ToolTipText_Lambda([this]()
-		{
-			return LogEntry->LogMetaData.AckingMessageId.IsSet() && CanScrollToAckLogFunc.Execute(*LogEntry->LogMetaData.AckingMessageId)
-				? LOCTEXT("Acked.Scroll.Enabled", "ACK received. Press to jump to the log that ACKs this log")
-				: LOCTEXT("Acked.Scroll.Disabled", "ACK received\nThis button is disabled because the corresponding ACK log is filtered out.");
-		})
-		[
-			SNew(SImage)
-			.DesiredSizeOverride(FVector2D(16.f, 16.f))
-			.Image(FConcertServerStyle::Get().GetBrush("Concert.Ack.Success"))
-		];
-}
+				return LogEntry->LogMetaData.AckingMessageId.IsSet() && CanScrollToAckLogFunc.Execute(*LogEntry->LogMetaData.AckingMessageId)
+					? LOCTEXT("Acked.Scroll.Enabled", "ACK received. Press to jump to the log that ACKs this log")
+					: LOCTEXT("Acked.Scroll.Disabled", "ACK received\nThis button is disabled because the corresponding ACK log is filtered out.");
+			})
+			[
+				SNew(SImage)
+				.DesiredSizeOverride(FVector2D(16.f, 16.f))
+				.Image(FConcertServerStyle::Get().GetBrush("Concert.Ack.Success"))
+			];
+	}
 
-TSharedRef<SWidget> SConcertTransportLogRow::CreateAckFailureWidget() const
-{
-	return SNew(SImage)
-		.DesiredSizeOverride(FVector2D(16.f, 16.f))
-		.Image(FConcertServerStyle::Get().GetBrush("Concert.Ack.Failure"))
-		.Visibility_Lambda([this]()
-		{
-			return LogEntry->LogMetaData.AckState == EConcertLogAckState::AckFailure
-				? EVisibility::Visible
-				: EVisibility::Hidden;
-		})
-		.ToolTipText(LOCTEXT("AckFailure", "No ACK received: retries exhausted."));
-}
+	TSharedRef<SWidget> SConcertTransportLogRow::CreateAckFailureWidget() const
+	{
+		return SNew(SImage)
+			.DesiredSizeOverride(FVector2D(16.f, 16.f))
+			.Image(FConcertServerStyle::Get().GetBrush("Concert.Ack.Failure"))
+			.Visibility_Lambda([this]()
+			{
+				return LogEntry->LogMetaData.AckState == EConcertLogAckState::AckFailure
+					? EVisibility::Visible
+					: EVisibility::Hidden;
+			})
+			.ToolTipText(LOCTEXT("AckFailure", "No ACK received: retries exhausted."));
+	}
 
-TSharedRef<SWidget> SConcertTransportLogRow::CreateAckNotNeededWidget() const
-{
-	return SNew(STextBlock)
-		.Text(LOCTEXT("NotApplicable", "n/a"))
-		.Visibility_Lambda([this]()
-		{
-			return LogEntry->LogMetaData.AckState == EConcertLogAckState::NotNeeded
-				? EVisibility::Visible
-				: EVisibility::Hidden;
-		})
-		.ToolTipText(LOCTEXT("Received", "No ACK needed."));
-}
+	TSharedRef<SWidget> SConcertTransportLogRow::CreateAckNotNeededWidget() const
+	{
+		return SNew(STextBlock)
+			.Text(LOCTEXT("NotApplicable", "n/a"))
+			.Visibility_Lambda([this]()
+			{
+				return LogEntry->LogMetaData.AckState == EConcertLogAckState::NotNeeded
+					? EVisibility::Visible
+					: EVisibility::Hidden;
+			})
+			.ToolTipText(LOCTEXT("Received", "No ACK needed."));
+	}
 
-TSharedRef<SWidget> SConcertTransportLogRow::CreateAckWidget() const
-{
-	return SNew(SComboButton)
-		.HasDownArrow(false)
-		.ToolTipText_Lambda([this]()
-		{
-			const bool bIsEnabled = LogEntry->LogMetaData.AckedMessageId && Algo::AnyOf(*LogEntry->LogMetaData.AckedMessageId, [this](const FGuid& MessageId)
+	TSharedRef<SWidget> SConcertTransportLogRow::CreateAckWidget() const
+	{
+		return SNew(SComboButton)
+			.HasDownArrow(false)
+			.ToolTipText_Lambda([this]()
+			{
+				const bool bIsEnabled = LogEntry->LogMetaData.AckedMessageId && Algo::AnyOf(*LogEntry->LogMetaData.AckedMessageId, [this](const FGuid& MessageId)
+					{
+					return CanScrollToAckedLogFunc.Execute(MessageId);
+				});
+				return bIsEnabled
+					? LOCTEXT("Ack.Scroll.Disabled", "This is an ACK. Press to jump to the log this log ACKs.")
+					: LOCTEXT("Ack.Scroll.Enabled", "This is an ACK.\nThis button is disabled because all logs this log ACKs are filtered out.");
+			})
+			.IsEnabled_Lambda([this]()
+			{
+				return LogEntry->LogMetaData.AckedMessageId && Algo::AnyOf(*LogEntry->LogMetaData.AckedMessageId, [this](const FGuid& MessageId)
 				{
-				return CanScrollToAckedLogFunc.Execute(MessageId);
-			});
-			return bIsEnabled
-				? LOCTEXT("Ack.Scroll.Disabled", "This is an ACK. Press to jump to the log this log ACKs.")
-				: LOCTEXT("Ack.Scroll.Enabled", "This is an ACK.\nThis button is disabled because all logs this log ACKs are filtered out.");
-		})
-		.IsEnabled_Lambda([this]()
-		{
-			return LogEntry->LogMetaData.AckedMessageId && Algo::AnyOf(*LogEntry->LogMetaData.AckedMessageId, [this](const FGuid& MessageId)
+					return CanScrollToAckedLogFunc.Execute(MessageId);
+				});
+			})
+			.OnGetMenuContent_Lambda([this]()
 			{
-				return CanScrollToAckedLogFunc.Execute(MessageId);
-			});
-		})
-		.OnGetMenuContent_Lambda([this]()
-		{
-			FMenuBuilder MenuBuilder(true, nullptr);
-			for (const FGuid& MessageId : *LogEntry->LogMetaData.AckedMessageId)
+				FMenuBuilder MenuBuilder(true, nullptr);
+				for (const FGuid& MessageId : *LogEntry->LogMetaData.AckedMessageId)
+				{
+					MenuBuilder.AddMenuEntry(
+						FText::FromString(MessageId.ToString(EGuidFormats::DigitsWithHyphens)),
+						CanScrollToAckedLogFunc.Execute(MessageId) ? FText() : LOCTEXT("Ack.Scroll.ComoboEntry.Disabled", "This log is filtered out."),
+						FSlateIcon(),
+						FUIAction(
+							FExecuteAction::CreateLambda([this, MessageId]()
+							{
+								ScrollToAckedLogFunc.Execute(MessageId);
+							}),
+							FCanExecuteAction::CreateLambda([this, MessageId] { return CanScrollToAckedLogFunc.Execute(MessageId); })),
+						NAME_None,
+						EUserInterfaceActionType::Button
+					);
+				}
+				return MenuBuilder.MakeWidget();
+			})
+			.Visibility_Lambda([this]()
 			{
-				MenuBuilder.AddMenuEntry(
-					FText::FromString(MessageId.ToString(EGuidFormats::DigitsWithHyphens)),
-					CanScrollToAckedLogFunc.Execute(MessageId) ? FText() : LOCTEXT("Ack.Scroll.ComoboEntry.Disabled", "This log is filtered out."),
-					FSlateIcon(),
-					FUIAction(
-						FExecuteAction::CreateLambda([this, MessageId]()
-						{
-							ScrollToAckedLogFunc.Execute(MessageId);
-						}),
-						FCanExecuteAction::CreateLambda([this, MessageId] { return CanScrollToAckedLogFunc.Execute(MessageId); })),
-					NAME_None,
-					EUserInterfaceActionType::Button
-				);
-			}
-			return MenuBuilder.MakeWidget();
-		})
-		.Visibility_Lambda([this]()
-		{
-			return LogEntry->LogMetaData.AckState == EConcertLogAckState::Ack
-				? EVisibility::Visible
-				: EVisibility::Hidden;
-		})
-		.ButtonContent()
-		[
-			SNew(SImage)
-			.DesiredSizeOverride(FVector2D(16.f, 16.f))
-			.Image(FConcertServerStyle::Get().GetBrush("Concert.Ack.Ack"))
-		];
+				return LogEntry->LogMetaData.AckState == EConcertLogAckState::Ack
+					? EVisibility::Visible
+					: EVisibility::Hidden;
+			})
+			.ButtonContent()
+			[
+				SNew(SImage)
+				.DesiredSizeOverride(FVector2D(16.f, 16.f))
+				.Image(FConcertServerStyle::Get().GetBrush("Concert.Ack.Ack"))
+			];
+	}
 }
 
 #undef LOCTEXT_NAMESPACE
diff --git a/Engine/Plugins/Developer/Concert/ConcertApp/MultiUserServer/Source/MultiUserServer/Private/Widgets/Clients/Logging/SConcertTransportLogRow.h b/Engine/Plugins/Developer/Concert/ConcertApp/MultiUserServer/Source/MultiUserServer/Private/Widgets/Clients/Logging/SConcertTransportLogRow.h
--- a/Engine/Plugins/Developer/Concert/ConcertApp/MultiUserServer/Source/MultiUserServer/Private/Widgets/Clients/Logging/SConcertTransportLogRow.h
+++ b/Engine/Plugins/Developer/Concert/ConcertApp/MultiUserServer/Source/MultiUserServer/Private/Widgets/Clients/Logging/SConcertTransportLogRow.h
@@ -7,17 +7,20 @@
 #include "Widgets/DeclarativeSyntaxSupport.h"
 #include "Widgets/Views/STableRow.h"
 
-class FConcertLogTokenizer;
 
-class SConcertTransportLogRow : public SMultiColumnTableRow<TSharedPtr<FConcertLogEntry>>
+namespace UE::MultiUserServer
 {
-public:
+	class FConcertLogTokenizer;
+	
+	class SConcertTransportLogRow : public SMultiColumnTableRow<TSharedPtr<FConcertLogEntry>>
+	{
+	public:
 
-	DECLARE_DELEGATE_OneParam(FScrollToLog, const FGuid& /*MessageId*/);
-	DECLARE_DELEGATE_RetVal_OneParam(bool, FCanScrollToLog, const FGuid& /*MessageId*/);
+		DECLARE_DELEGATE_OneParam(FScrollToLog, const FGuid& /*MessageId*/);
+		DECLARE_DELEGATE_RetVal_OneParam(bool, FCanScrollToLog, const FGuid& /*MessageId*/);
 
-	SLATE_BEGIN_ARGS(SConcertTransportLogRow)
-	{}
+		SLATE_BEGIN_ARGS(SConcertTransportLogRow)
+		{}
 		SLATE_ATTRIBUTE(FText, HighlightText)
 		SLATE_ARGUMENT(FLinearColor, AvatarColor)
 		SLATE_EVENT(FCanScrollToLog, CanScrollToAckLog)
@@ -28,31 +31,32 @@ public:
 
 	void Construct(const FArguments& InArgs, TSharedPtr<FConcertLogEntry> InLogEntry, const TSharedRef<STableViewBase>& InOwnerTableView, TSharedRef<FConcertLogTokenizer> InTokenizer);
 
-	//~ Begin SMultiColumnTableRow Interface
-	virtual TSharedRef<SWidget> GenerateWidgetForColumn(const FName& ColumnName) override;
-	//~ End SMultiColumnTableRow Interface
+		//~ Begin SMultiColumnTableRow Interface
+		virtual TSharedRef<SWidget> GenerateWidgetForColumn(const FName& ColumnName) override;
+		//~ End SMultiColumnTableRow Interface
 	
-private:
+		private:
 
-	TSharedPtr<FConcertLogEntry> LogEntry;
-	/** Used to convert some members into display strings */
-	TSharedPtr<FConcertLogTokenizer> Tokenizer;
+		TSharedPtr<FConcertLogEntry> LogEntry;
+		/** Used to convert some members into display strings */
+		TSharedPtr<FConcertLogTokenizer> Tokenizer;
 
-	TAttribute<FText> HighlightText;
+		TAttribute<FText> HighlightText;
 
-	FLinearColor AvatarColor;
-	FCanScrollToLog CanScrollToAckLogFunc;
-	FCanScrollToLog CanScrollToAckedLogFunc;
-	FScrollToLog ScrollToAckLogFunc;
-	FScrollToLog ScrollToAckedLogFunc;
+		FLinearColor AvatarColor;
+		FCanScrollToLog CanScrollToAckLogFunc;
+		FCanScrollToLog CanScrollToAckedLogFunc;
+		FScrollToLog ScrollToAckLogFunc;
+		FScrollToLog ScrollToAckedLogFunc;
 
-	TSharedRef<SWidget> CreateDefaultColumn(const FName& PropertyName);
+		TSharedRef<SWidget> CreateDefaultColumn(const FName& PropertyName);
 	
-	TSharedRef<SWidget> CreateAckColumn(const FName& PropertyName);
-	TSharedRef<SWidget> CreateAckInProgressWidget() const;
-	TSharedRef<SWidget> CreateAckSuccessWidget() const;
-	TSharedRef<SWidget> CreateAckFailureWidget() const;
-	TSharedRef<SWidget> CreateAckNotNeededWidget() const;
-	TSharedRef<SWidget> CreateAckWidget() const;
+		TSharedRef<SWidget> CreateAckColumn(const FName& PropertyName);
+		TSharedRef<SWidget> CreateAckInProgressWidget() const;
+		TSharedRef<SWidget> CreateAckSuccessWidget() const;
+		TSharedRef<SWidget> CreateAckFailureWidget() const;
+		TSharedRef<SWidget> CreateAckNotNeededWidget() const;
+		TSharedRef<SWidget> CreateAckWidget() const;
 	
-};
+	};
+}
\ No newline at end of file
diff --git a/Engine/Plugins/Developer/Concert/ConcertApp/MultiUserServer/Source/MultiUserServer/Private/Widgets/Clients/Logging/Util/ConcertLogTokenizer.cpp b/Engine/Plugins/Developer/Concert/ConcertApp/MultiUserServer/Source/MultiUserServer/Private/Widgets/Clients/Logging/Util/ConcertLogTokenizer.cpp
--- a/Engine/Plugins/Developer/Concert/ConcertApp/MultiUserServer/Source/MultiUserServer/Private/Widgets/Clients/Logging/Util/ConcertLogTokenizer.cpp
+++ b/Engine/Plugins/Developer/Concert/ConcertApp/MultiUserServer/Source/MultiUserServer/Private/Widgets/Clients/Logging/Util/ConcertLogTokenizer.cpp
@@ -9,71 +9,74 @@
 #include "Settings/ConcertTransportLogSettings.h"
 #include "Widgets/Clients/Util/EndpointToUserNameCache.h"
 
-FConcertLogTokenizer::FConcertLogTokenizer(TSharedRef<FEndpointToUserNameCache> EndpointInfoGetter)
-	: EndpointInfoGetter(EndpointInfoGetter)
+namespace UE::MultiUserServer
 {
-	TokenizerFunctions = {
-			{ FConcertLog::StaticStruct()->FindPropertyByName(GET_MEMBER_NAME_CHECKED(FConcertLog, MessageId)), [this](const FConcertLog& Log) { return TokenizeMessageId(Log); } },
-			{ FConcertLog::StaticStruct()->FindPropertyByName(GET_MEMBER_NAME_CHECKED(FConcertLog, Timestamp)), [this](const FConcertLog& Log) { return TokenizeTimestamp(Log); } },
-			{ FConcertLog::StaticStruct()->FindPropertyByName(GET_MEMBER_NAME_CHECKED(FConcertLog, MessageTypeName)), [this](const FConcertLog& Log) { return TokenizeMessageTypeName(Log); } },
-			{ FConcertLog::StaticStruct()->FindPropertyByName(GET_MEMBER_NAME_CHECKED(FConcertLog, CustomPayloadUncompressedByteSize)), [this](const FConcertLog& Log) { return TokenizeCustomPayloadUncompressedByteSize(Log); } },
-			{ FConcertLog::StaticStruct()->FindPropertyByName(GET_MEMBER_NAME_CHECKED(FConcertLog, OriginEndpointId)), [this](const FConcertLog& Log) { return TokenizeOriginEndpointId(Log); } },
-			{ FConcertLog::StaticStruct()->FindPropertyByName(GET_MEMBER_NAME_CHECKED(FConcertLog, DestinationEndpointId)), [this](const FConcertLog& Log) { return TokenizeDestinationEndpointId(Log); } }
-	};
-}
+	FConcertLogTokenizer::FConcertLogTokenizer(TSharedRef<FEndpointToUserNameCache> EndpointInfoGetter)
+		: EndpointInfoGetter(EndpointInfoGetter)
+	{
+		TokenizerFunctions = {
+				{ FConcertLog::StaticStruct()->FindPropertyByName(GET_MEMBER_NAME_CHECKED(FConcertLog, MessageId)), [this](const FConcertLog& Log) { return TokenizeMessageId(Log); } },
+				{ FConcertLog::StaticStruct()->FindPropertyByName(GET_MEMBER_NAME_CHECKED(FConcertLog, Timestamp)), [this](const FConcertLog& Log) { return TokenizeTimestamp(Log); } },
+				{ FConcertLog::StaticStruct()->FindPropertyByName(GET_MEMBER_NAME_CHECKED(FConcertLog, MessageTypeName)), [this](const FConcertLog& Log) { return TokenizeMessageTypeName(Log); } },
+				{ FConcertLog::StaticStruct()->FindPropertyByName(GET_MEMBER_NAME_CHECKED(FConcertLog, CustomPayloadUncompressedByteSize)), [this](const FConcertLog& Log) { return TokenizeCustomPayloadUncompressedByteSize(Log); } },
+				{ FConcertLog::StaticStruct()->FindPropertyByName(GET_MEMBER_NAME_CHECKED(FConcertLog, OriginEndpointId)), [this](const FConcertLog& Log) { return TokenizeOriginEndpointId(Log); } },
+				{ FConcertLog::StaticStruct()->FindPropertyByName(GET_MEMBER_NAME_CHECKED(FConcertLog, DestinationEndpointId)), [this](const FConcertLog& Log) { return TokenizeDestinationEndpointId(Log); } }
+		};
+	}
 
-FString FConcertLogTokenizer::Tokenize(const FConcertLog& Data, const FProperty& ConcertLogProperty) const
-{
-	if (const FTokenizeFunc* CustomTokenizer = TokenizerFunctions.Find(&ConcertLogProperty))
+	FString FConcertLogTokenizer::Tokenize(const FConcertLog& Data, const FProperty& ConcertLogProperty) const
 	{
-		return (*CustomTokenizer)(Data);
+		if (const FTokenizeFunc* CustomTokenizer = TokenizerFunctions.Find(&ConcertLogProperty))
+		{
+			return (*CustomTokenizer)(Data);
+		}
+		return TokenizeUsingPropertyExport(&Data, ConcertLogProperty);
 	}
-	return TokenizeUsingPropertyExport(&Data, ConcertLogProperty);
-}
 
-FString FConcertLogTokenizer::Tokenize(const FConcertLogMetadata& Data, const FProperty& ConcertLogMetadataProperty) const
-{
-	return TokenizeUsingPropertyExport(&Data, ConcertLogMetadataProperty);
-}
+	FString FConcertLogTokenizer::Tokenize(const FConcertLogMetadata& Data, const FProperty& ConcertLogMetadataProperty) const
+	{
+		return TokenizeUsingPropertyExport(&Data, ConcertLogMetadataProperty);
+	}
 
-FString FConcertLogTokenizer::TokenizeMessageId(const FConcertLog& Data) const
-{
-	return Data.MessageId.ToString(EGuidFormats::DigitsWithHyphens);
-}
+	FString FConcertLogTokenizer::TokenizeMessageId(const FConcertLog& Data) const
+	{
+		return Data.MessageId.ToString(EGuidFormats::DigitsWithHyphens);
+	}
 
-FString FConcertLogTokenizer::TokenizeTimestamp(const FConcertLog& Data) const
-{
-	return ConcertFrontendUtils::FormatTime(Data.Timestamp, UConcertTransportLogSettings::GetSettings()->TimestampTimeFormat).ToString();
-}
+	FString FConcertLogTokenizer::TokenizeTimestamp(const FConcertLog& Data) const
+	{
+		return ConcertFrontendUtils::FormatTime(Data.Timestamp, UConcertTransportLogSettings::GetSettings()->TimestampTimeFormat).ToString();
+	}
 
-FString FConcertLogTokenizer::TokenizeMessageTypeName(const FConcertLog& Data) const
-{
-	return UE::MultiUserServer::MessageTypeUtils::SanitizeMessageTypeName(Data.MessageTypeName);
-}
+	FString FConcertLogTokenizer::TokenizeMessageTypeName(const FConcertLog& Data) const
+	{
+		return UE::MultiUserServer::MessageTypeUtils::SanitizeMessageTypeName(Data.MessageTypeName);
+	}
 
-FString FConcertLogTokenizer::TokenizeCustomPayloadUncompressedByteSize(const FConcertLog& Data) const
-{
-	// So changes to the type are automatically propagated here
-	const FNumericUnit<int32> DisplayUnit = FUnitConversion::QuantizeUnitsToBestFit(Data.CustomPayloadUncompressedByteSize, EUnit::Bytes);
-	return FString::Printf(TEXT("%d %s"), DisplayUnit.Value, FUnitConversion::GetUnitDisplayString(DisplayUnit.Units));
-}
+	FString FConcertLogTokenizer::TokenizeCustomPayloadUncompressedByteSize(const FConcertLog& Data) const
+	{
+		// So changes to the type are automatically propagated here
+		const FNumericUnit<int32> DisplayUnit = FUnitConversion::QuantizeUnitsToBestFit(Data.CustomPayloadUncompressedByteSize, EUnit::Bytes);
+		return FString::Printf(TEXT("%d %s"), DisplayUnit.Value, FUnitConversion::GetUnitDisplayString(DisplayUnit.Units));
+	}
 
-FString FConcertLogTokenizer::TokenizeOriginEndpointId(const FConcertLog& Data) const
-{
-	return EndpointInfoGetter->GetEndpointDisplayString(Data.OriginEndpointId);
-}
+	FString FConcertLogTokenizer::TokenizeOriginEndpointId(const FConcertLog& Data) const
+	{
+		return EndpointInfoGetter->GetEndpointDisplayString(Data.OriginEndpointId);
+	}
 
-FString FConcertLogTokenizer::TokenizeDestinationEndpointId(const FConcertLog& Data) const
-{
-	return EndpointInfoGetter->GetEndpointDisplayString(Data.DestinationEndpointId);
-}
+	FString FConcertLogTokenizer::TokenizeDestinationEndpointId(const FConcertLog& Data) const
+	{
+		return EndpointInfoGetter->GetEndpointDisplayString(Data.DestinationEndpointId);
+	}
 
-FString FConcertLogTokenizer::TokenizeUsingPropertyExport(const void* ContainerPtr, const FProperty& ConcertLogProperty) const
-{
-	FString Exported;
-	const void* ValuePtr = ConcertLogProperty.ContainerPtrToValuePtr<void>(ContainerPtr);
-	const void* DeltaPtr = ValuePtr; // We have no real delta - in this case the API expects the same ptr
-	const bool bSuccess = ConcertLogProperty.ExportText_Direct(Exported, ValuePtr, DeltaPtr, nullptr, PPF_ExternalEditor);
-	check(bSuccess);
-	return Exported;
+	FString FConcertLogTokenizer::TokenizeUsingPropertyExport(const void* ContainerPtr, const FProperty& ConcertLogProperty) const
+	{
+		FString Exported;
+		const void* ValuePtr = ConcertLogProperty.ContainerPtrToValuePtr<void>(ContainerPtr);
+		const void* DeltaPtr = ValuePtr; // We have no real delta - in this case the API expects the same ptr
+		const bool bSuccess = ConcertLogProperty.ExportText_Direct(Exported, ValuePtr, DeltaPtr, nullptr, PPF_ExternalEditor);
+		check(bSuccess);
+		return Exported;
+	}
 }
diff --git a/Engine/Plugins/Developer/Concert/ConcertApp/MultiUserServer/Source/MultiUserServer/Private/Widgets/Clients/Logging/Util/ConcertLogTokenizer.h b/Engine/Plugins/Developer/Concert/ConcertApp/MultiUserServer/Source/MultiUserServer/Private/Widgets/Clients/Logging/Util/ConcertLogTokenizer.h
--- a/Engine/Plugins/Developer/Concert/ConcertApp/MultiUserServer/Source/MultiUserServer/Private/Widgets/Clients/Logging/Util/ConcertLogTokenizer.h
+++ b/Engine/Plugins/Developer/Concert/ConcertApp/MultiUserServer/Source/MultiUserServer/Private/Widgets/Clients/Logging/Util/ConcertLogTokenizer.h
@@ -8,35 +8,39 @@ class FEndpointToUserNameCache;
 struct FConcertLog;
 struct FConcertLogMetadata;
 
-/** Converts members of FConcertLog and FConcertLogMetadata into a string. Used e.g. to make search respect the display settings. */
-class FConcertLogTokenizer
+namespace UE::MultiUserServer
 {
-public:
+	/** Converts members of FConcertLog and FConcertLogMetadata into a string. Used e.g. to make search respect the display settings. */
+	class FConcertLogTokenizer
+	{
+	public:
 
-	FConcertLogTokenizer(TSharedRef<FEndpointToUserNameCache> EndpointInfoGetter);
+		FConcertLogTokenizer(TSharedRef<FEndpointToUserNameCache> EndpointInfoGetter);
 
-	/** Tokenizes a property of FConcertLog into a string */
-	FString Tokenize(const FConcertLog& Data, const FProperty& ConcertLogProperty) const;
-	/** Tokenizes a property of FConcertLogMetadata into a string */
-	FString Tokenize(const FConcertLogMetadata& Data, const FProperty& ConcertLogMetadataProperty) const;
+		/** Tokenizes a property of FConcertLog into a string */
+		FString Tokenize(const FConcertLog& Data, const FProperty& ConcertLogProperty) const;
+		/** Tokenizes a property of FConcertLogMetadata into a string */
+		FString Tokenize(const FConcertLogMetadata& Data, const FProperty& ConcertLogMetadataProperty) const;
 
-	// ContainerPtr Tokenizers
-	FString TokenizeMessageId(const FConcertLog& Data) const;
-	FString TokenizeTimestamp(const FConcertLog& Data) const;
-	FString TokenizeMessageTypeName(const FConcertLog& Data) const;
-	FString TokenizeCustomPayloadUncompressedByteSize(const FConcertLog& Data) const;
-	FString TokenizeOriginEndpointId(const FConcertLog& Data) const;
-	FString TokenizeDestinationEndpointId(const FConcertLog& Data) const;
+		// ContainerPtr Tokenizers
+		FString TokenizeMessageId(const FConcertLog& Data) const;
+		FString TokenizeTimestamp(const FConcertLog& Data) const;
+		FString TokenizeMessageTypeName(const FConcertLog& Data) const;
+		FString TokenizeCustomPayloadUncompressedByteSize(const FConcertLog& Data) const;
+		FString TokenizeOriginEndpointId(const FConcertLog& Data) const;
+		FString TokenizeDestinationEndpointId(const FConcertLog& Data) const;
 
-private:
+	private:
 
-	using FTokenizeFunc = TFunction<FString(const FConcertLog&)>;
+		using FTokenizeFunc = TFunction<FString(const FConcertLog&)>;
 
-	/** Override functions for tokenizing certain properties */
-	TMap<const FProperty*, FTokenizeFunc> TokenizerFunctions;
+		/** Override functions for tokenizing certain properties */
+		TMap<const FProperty*, FTokenizeFunc> TokenizerFunctions;
 
-	/** Used so we can look up client and server info (even after client has disconnected) */
-	TSharedRef<FEndpointToUserNameCache> EndpointInfoGetter;
+		/** Used so we can look up client and server info (even after client has disconnected) */
+		TSharedRef<FEndpointToUserNameCache> EndpointInfoGetter;
 	
-	FString TokenizeUsingPropertyExport(const void* ContainerPtr, const FProperty& ConcertLogProperty) const;
-};
+		FString TokenizeUsingPropertyExport(const void* ContainerPtr, const FProperty& ConcertLogProperty) const;
+	};
+}
+
diff --git a/Engine/Plugins/Developer/Concert/ConcertApp/MultiUserServer/Source/MultiUserServer/Private/Widgets/Clients/PackageTransmission/Filter/FrontendPackageTransmissionFilter.h b/Engine/Plugins/Developer/Concert/ConcertApp/MultiUserServer/Source/MultiUserServer/Private/Widgets/Clients/PackageTransmission/Filter/FrontendPackageTransmissionFilter.h
deleted file mode 100644
diff --git a/Engine/Plugins/Developer/Concert/ConcertApp/MultiUserServer/Source/MultiUserServer/Private/Widgets/Clients/PackageTransmission/Filter/FrontendPackageTransmissionFilter_TextSearch.cpp b/Engine/Plugins/Developer/Concert/ConcertApp/MultiUserServer/Source/MultiUserServer/Private/Widgets/Clients/PackageTransmission/Filter/FrontendPackageTransmissionFilter_TextSearch.cpp
deleted file mode 100644
diff --git a/Engine/Plugins/Developer/Concert/ConcertApp/MultiUserServer/Source/MultiUserServer/Private/Widgets/Clients/PackageTransmission/Filter/FrontendPackageTransmissionFilter_TextSearch.h b/Engine/Plugins/Developer/Concert/ConcertApp/MultiUserServer/Source/MultiUserServer/Private/Widgets/Clients/PackageTransmission/Filter/FrontendPackageTransmissionFilter_TextSearch.h
deleted file mode 100644
diff --git a/Engine/Plugins/Developer/Concert/ConcertApp/MultiUserServer/Source/MultiUserServer/Private/Widgets/Clients/PackageTransmission/Filter/PackageTransmissionFilter.h b/Engine/Plugins/Developer/Concert/ConcertApp/MultiUserServer/Source/MultiUserServer/Private/Widgets/Clients/PackageTransmission/Filter/PackageTransmissionFilter.h
deleted file mode 100644
diff --git a/Engine/Plugins/Developer/Concert/ConcertApp/MultiUserServer/Source/MultiUserServer/Private/Widgets/Clients/PackageTransmission/Filter/PackageTransmissionFilterTypes.h b/Engine/Plugins/Developer/Concert/ConcertApp/MultiUserServer/Source/MultiUserServer/Private/Widgets/Clients/PackageTransmission/Filter/PackageTransmissionFilterTypes.h
new file mode 100644
--- /dev/null
+++ b/Engine/Plugins/Developer/Concert/ConcertApp/MultiUserServer/Source/MultiUserServer/Private/Widgets/Clients/PackageTransmission/Filter/PackageTransmissionFilterTypes.h
@@ -0,0 +1,17 @@
+// Copyright Epic Games, Inc. All Rights Reserved.
+
+#pragma once
+
+#include "CoreMinimal.h"
+
+namespace UE::MultiUserServer
+{
+	struct FPackageTransmissionEntry;
+	template<typename TFilterType> class TConcertFilter;
+	template<typename TFilterType> class TConcertFrontendFilter;
+
+	using FPackageTransmissionFilter = TConcertFilter<const FPackageTransmissionEntry&>;
+	using FFrontendPackageTransmissionFilter = TConcertFrontendFilter<const FPackageTransmissionEntry&>;
+}
+
+
diff --git a/Engine/Plugins/Developer/Concert/ConcertApp/MultiUserServer/Source/MultiUserServer/Private/Widgets/Clients/PackageTransmission/Filter/PackageTransmissionFilter_FrontendRoot.cpp b/Engine/Plugins/Developer/Concert/ConcertApp/MultiUserServer/Source/MultiUserServer/Private/Widgets/Clients/PackageTransmission/Filter/PackageTransmissionFilter_FrontendRoot.cpp
--- a/Engine/Plugins/Developer/Concert/ConcertApp/MultiUserServer/Source/MultiUserServer/Private/Widgets/Clients/PackageTransmission/Filter/PackageTransmissionFilter_FrontendRoot.cpp
+++ b/Engine/Plugins/Developer/Concert/ConcertApp/MultiUserServer/Source/MultiUserServer/Private/Widgets/Clients/PackageTransmission/Filter/PackageTransmissionFilter_FrontendRoot.cpp
@@ -2,88 +2,14 @@
 
 #include "PackageTransmissionFilter_FrontendRoot.h"
 
-#include "Widgets/Clients/PackageTransmission/Model/PackageTransmissionEntry.h"
-#include "Widgets/Clients/PackageTransmission/Filter/FrontendPackageTransmissionFilter_TextSearch.h"
+#include "Widgets/Clients/PackageTransmission/Filter/PackageTransmissionFrontendFilter_TextSearch.h"
 
-#include "Algo/AllOf.h"
-#include "Widgets/SBoxPanel.h"
 
 namespace UE::MultiUserServer
 {
-	FPackageTransmissionFilter_FrontendRoot::FPackageTransmissionFilter_FrontendRoot(
-		TSharedRef<FPackageTransmissionEntryTokenizer> Tokenizer,
-		TArray<TSharedRef<FFrontendPackageTransmissionFilter>> InCustomFilters,
-		const TArray<TSharedRef<FPackageTransmissionFilter>>& NonVisualFilters)
-		: TextSearchFilter(MakeShared<FFrontendPackageTransmissionFilter_TextSearch>(MoveTemp(Tokenizer)))
-		, FrontendFilters(MoveTemp(InCustomFilters))
-		, AllFilters(FrontendFilters)
-	{
-		AllFilters.Add(TextSearchFilter);
-		for (const TSharedRef<FPackageTransmissionFilter>& NonVisual : NonVisualFilters)
-		{
-			AllFilters.Add(NonVisual);
-		}
-	
-		for (const TSharedRef<FPackageTransmissionFilter>& Filter : AllFilters)
-		{
-			Filter->OnChanged().AddLambda([this]()
-			{
-				OnChanged().Broadcast();
-			});
-		}
-	}
-
-	TSharedRef<SWidget> FPackageTransmissionFilter_FrontendRoot::BuildFilterWidgets() const
-	{
-		return SNew(SVerticalBox)
-
-			// Search bar
-			+SVerticalBox::Slot()
-			.AutoHeight()
-			[
-				TextSearchFilter->GetFilterWidget()
-			]
-
-			+SVerticalBox::Slot()
-			.AutoHeight()
-			.Padding(0, 2)
-			[
-				BuildCustomFilterListWidget()
-			];
-	}
-
-	bool FPackageTransmissionFilter_FrontendRoot::PassesFilter(const FPackageTransmissionEntry& InItem) const
-	{
-		return Algo::AllOf(
-			AllFilters,
-			[&InItem](const TSharedRef<FPackageTransmissionFilter>& AndFilter){ return AndFilter->PassesFilter(InItem); }
-		);
-	}
-
-	TSharedRef<SWidget> FPackageTransmissionFilter_FrontendRoot::BuildCustomFilterListWidget() const
-	{
-		TSharedRef<SHorizontalBox> Box = SNew(SHorizontalBox);
-
-		bool bIsFirst = true;
-		for (const TSharedRef<FFrontendPackageTransmissionFilter>& Filter : FrontendFilters)
-		{
-			const FMargin Margin = bIsFirst ? FMargin() : FMargin(8, 0, 0, 0);
-			bIsFirst = false;
-			Box->AddSlot()
-				.AutoWidth()
-				.Padding(Margin)
-				.VAlign(VAlign_Center)
-				[
-					Filter->GetFilterWidget()
-				];
-		}
-
-		return Box;
-	}
-
 	TSharedRef<FPackageTransmissionFilter_FrontendRoot> MakeFilter(TSharedRef<FPackageTransmissionEntryTokenizer> Tokenizer)
 	{
-		TArray<TSharedRef<FFrontendPackageTransmissionFilter>> FrontendFilters;
+		TArray<TSharedRef<TConcertFrontendFilter<const FPackageTransmissionEntry&>>> FrontendFilters;
 		return MakeShared<FPackageTransmissionFilter_FrontendRoot>(MoveTemp(Tokenizer), FrontendFilters);
 	}
 }
diff --git a/Engine/Plugins/Developer/Concert/ConcertApp/MultiUserServer/Source/MultiUserServer/Private/Widgets/Clients/PackageTransmission/Filter/PackageTransmissionFilter_FrontendRoot.h b/Engine/Plugins/Developer/Concert/ConcertApp/MultiUserServer/Source/MultiUserServer/Private/Widgets/Clients/PackageTransmission/Filter/PackageTransmissionFilter_FrontendRoot.h
--- a/Engine/Plugins/Developer/Concert/ConcertApp/MultiUserServer/Source/MultiUserServer/Private/Widgets/Clients/PackageTransmission/Filter/PackageTransmissionFilter_FrontendRoot.h
+++ b/Engine/Plugins/Developer/Concert/ConcertApp/MultiUserServer/Source/MultiUserServer/Private/Widgets/Clients/PackageTransmission/Filter/PackageTransmissionFilter_FrontendRoot.h
@@ -3,50 +3,31 @@
 #pragma once
 
 #include "CoreMinimal.h"
-#include "PackageTransmissionFilter.h"
+#include "PackageTransmissionFrontendFilter_TextSearch.h"
+#include "PackageTransmissionFilterTypes.h"
+#include "Widgets/Util/Filter/ConcertRootFrontendFilter.h"
 
 class FEndpointToUserNameCache;
 class SWidget;
 
 namespace UE::MultiUserServer
 {
-	class FFrontendPackageTransmissionFilter;
-	class FFrontendPackageTransmissionFilter_TextSearch;
+	class FPackageTransmissionFrontendFilter_TextSearch;
 	class FPackageTransmissionEntryTokenizer;
 
-	/** A filter that contains multiple UI filters */
-	class FPackageTransmissionFilter_FrontendRoot : public FPackageTransmissionFilter, public TSharedFromThis<FPackageTransmissionFilter_FrontendRoot>
+	/** Root filter displayed in package transmission UI */
+	class FPackageTransmissionFilter_FrontendRoot : public TConcertFrontendRootFilter<const FPackageTransmissionEntry&, FPackageTransmissionFrontendFilter_TextSearch>
 	{
+		using Super = TConcertFrontendRootFilter<const FPackageTransmissionEntry&, FPackageTransmissionFrontendFilter_TextSearch>;
 	public:
 
 		FPackageTransmissionFilter_FrontendRoot(
 			TSharedRef<FPackageTransmissionEntryTokenizer> Tokenizer,
-			TArray<TSharedRef<FFrontendPackageTransmissionFilter>> InCustomFilters,
-			const TArray<TSharedRef<FPackageTransmissionFilter>>& NonVisualFilters = {}
-			);
-
-		/** Builds the widget view for all contained filters */
-		TSharedRef<SWidget> BuildFilterWidgets() const;
-		
-		//~ Begin IFilter Interface
-		virtual bool PassesFilter(const FPackageTransmissionEntry& InItem) const override;
-		//~ End IFilter Interface
-
-		FORCEINLINE const TSharedRef<FFrontendPackageTransmissionFilter_TextSearch>& GetTextSearchFilter() const { return TextSearchFilter; }
-
-	private:
-
-		/** The text search filter. Also in FrontendFilters. Separate variable to build search bar in new line. */
-		TSharedRef<FFrontendPackageTransmissionFilter_TextSearch> TextSearchFilter;
-
-		/** AllFilters without special filters we have as properties above, such as TextSearchFilter. */
-		TArray<TSharedRef<FFrontendPackageTransmissionFilter>> FrontendFilters;
-		
-		/** Filters that are combined using logical AND. */
-		TArray<TSharedRef<FPackageTransmissionFilter>> AllFilters;
-
-		/** Builds the widgets that go under the text */
-		TSharedRef<SWidget> BuildCustomFilterListWidget() const;
+			TArray<TSharedRef<Super::FConcertFrontendFilter>> FrontendFilters,
+			const TArray<TSharedRef<Super::FConcertFilter>>& NonVisualFilters = {}
+			)
+			: Super(MakeShared<FPackageTransmissionFrontendFilter_TextSearch>(MoveTemp(Tokenizer)), MoveTemp(FrontendFilters), NonVisualFilters)
+		{}
 	};
 
 	/** Creates a filter for the global filter log window */
diff --git a/Engine/Plugins/Developer/Concert/ConcertApp/MultiUserServer/Source/MultiUserServer/Private/Widgets/Clients/PackageTransmission/Filter/PackageTransmissionFrontendFilter_TextSearch.h b/Engine/Plugins/Developer/Concert/ConcertApp/MultiUserServer/Source/MultiUserServer/Private/Widgets/Clients/PackageTransmission/Filter/PackageTransmissionFrontendFilter_TextSearch.h
new file mode 100644
--- /dev/null
+++ b/Engine/Plugins/Developer/Concert/ConcertApp/MultiUserServer/Source/MultiUserServer/Private/Widgets/Clients/PackageTransmission/Filter/PackageTransmissionFrontendFilter_TextSearch.h
@@ -0,0 +1,49 @@
+// Copyright Epic Games, Inc. All Rights Reserved.
+
+#pragma once
+
+#include "CoreMinimal.h"
+#include "PackageTransmissionFilterTypes.h"
+#include "Widgets/Clients/PackageTransmission/Util/PackageTransmissionEntryTokenizer.h"
+#include "Widgets/Util/Filter/ConcertFrontendFilter_TextSearch.h"
+
+namespace UE::MultiUserServer
+{
+	class FPackageTransmissionFilter_TextSearch : public TConcertFilter_TextSearch<const FPackageTransmissionEntry&>
+	{
+	public:
+		
+		FPackageTransmissionFilter_TextSearch(TSharedRef<FPackageTransmissionEntryTokenizer> Tokenizer)
+			: Tokenizer(MoveTemp(Tokenizer))
+		{}
+
+	protected:
+
+		virtual void GenerateSearchTerms(const FPackageTransmissionEntry& InItem, TArray<FString>& OutTerms) const override
+		{
+			OutTerms.Add(Tokenizer->TokenizeTime(InItem));
+			OutTerms.Add(Tokenizer->TokenizeOrigin(InItem));
+			OutTerms.Add(Tokenizer->TokenizeDestination(InItem));
+			OutTerms.Add(Tokenizer->TokenizeSize(InItem));
+			OutTerms.Add(Tokenizer->TokenizeRevision(InItem));
+			OutTerms.Add(Tokenizer->TokenizePackagePath(InItem));
+			OutTerms.Add(Tokenizer->TokenizePackageName(InItem));
+		}
+	
+	private:
+		
+		/** Helps in converting FConcertLog members into search terms */
+		TSharedRef<FPackageTransmissionEntryTokenizer> Tokenizer;
+	};
+	
+    class FPackageTransmissionFrontendFilter_TextSearch : public TConcertFrontendFilter_TextSearch<FPackageTransmissionFilter_TextSearch, const FPackageTransmissionEntry&>
+    {
+		using Super = TConcertFrontendFilter_TextSearch<FPackageTransmissionFilter_TextSearch, const FPackageTransmissionEntry&>;
+    public:
+    	
+    	FPackageTransmissionFrontendFilter_TextSearch(TSharedRef<FPackageTransmissionEntryTokenizer> Tokenizer)
+    		: Super(MoveTemp(Tokenizer))
+    	{}
+    };
+}
+
diff --git a/Engine/Plugins/Developer/Concert/ConcertApp/MultiUserServer/Source/MultiUserServer/Private/Widgets/Clients/PackageTransmission/Model/FilteredPackageTransmissionModel.cpp b/Engine/Plugins/Developer/Concert/ConcertApp/MultiUserServer/Source/MultiUserServer/Private/Widgets/Clients/PackageTransmission/Model/FilteredPackageTransmissionModel.cpp
--- a/Engine/Plugins/Developer/Concert/ConcertApp/MultiUserServer/Source/MultiUserServer/Private/Widgets/Clients/PackageTransmission/Model/FilteredPackageTransmissionModel.cpp
+++ b/Engine/Plugins/Developer/Concert/ConcertApp/MultiUserServer/Source/MultiUserServer/Private/Widgets/Clients/PackageTransmission/Model/FilteredPackageTransmissionModel.cpp
@@ -2,7 +2,7 @@
 
 #include "FilteredPackageTransmissionModel.h"
 
-#include "Widgets/Clients/PackageTransmission/Filter/PackageTransmissionFilter.h"
+#include "Widgets/Util/Filter/ConcertFilter.h"
 
 namespace UE::MultiUserServer
 {
diff --git a/Engine/Plugins/Developer/Concert/ConcertApp/MultiUserServer/Source/MultiUserServer/Private/Widgets/Clients/PackageTransmission/Model/FilteredPackageTransmissionModel.h b/Engine/Plugins/Developer/Concert/ConcertApp/MultiUserServer/Source/MultiUserServer/Private/Widgets/Clients/PackageTransmission/Model/FilteredPackageTransmissionModel.h
--- a/Engine/Plugins/Developer/Concert/ConcertApp/MultiUserServer/Source/MultiUserServer/Private/Widgets/Clients/PackageTransmission/Model/FilteredPackageTransmissionModel.h
+++ b/Engine/Plugins/Developer/Concert/ConcertApp/MultiUserServer/Source/MultiUserServer/Private/Widgets/Clients/PackageTransmission/Model/FilteredPackageTransmissionModel.h
@@ -4,12 +4,10 @@
 
 #include "CoreMinimal.h"
 #include "IPackageTransmissionEntrySource.h"
-#include "Widgets/Clients/PackageTransmission/Filter/PackageTransmissionFilter.h"
+#include "Widgets/Clients/PackageTransmission/Filter/PackageTransmissionFilterTypes.h"
 
 namespace UE::MultiUserServer
 {
-	class FPackageTransmissionFilter;
-
 	class FFilteredPackageTransmissionModel : public IPackageTransmissionEntrySource
 	{
 	public:
diff --git a/Engine/Plugins/Developer/Concert/ConcertApp/MultiUserServer/Source/MultiUserServer/Private/Widgets/Clients/PackageTransmission/SPackageTransmissionView.cpp b/Engine/Plugins/Developer/Concert/ConcertApp/MultiUserServer/Source/MultiUserServer/Private/Widgets/Clients/PackageTransmission/SPackageTransmissionView.cpp
--- a/Engine/Plugins/Developer/Concert/ConcertApp/MultiUserServer/Source/MultiUserServer/Private/Widgets/Clients/PackageTransmission/SPackageTransmissionView.cpp
+++ b/Engine/Plugins/Developer/Concert/ConcertApp/MultiUserServer/Source/MultiUserServer/Private/Widgets/Clients/PackageTransmission/SPackageTransmissionView.cpp
@@ -3,7 +3,7 @@
 #include "SPackageTransmissionView.h"
 
 #include "SPackageTransmissionTable.h"
-#include "Filter/FrontendPackageTransmissionFilter_TextSearch.h"
+#include "Filter/PackageTransmissionFrontendFilter_TextSearch.h"
 #include "Filter/PackageTransmissionFilter_FrontendRoot.h"
 #include "Model/FilteredPackageTransmissionModel.h"
 
diff --git a/Engine/Plugins/Developer/Concert/ConcertApp/MultiUserServer/Source/MultiUserServer/Private/Widgets/Clients/SConcertClientsTabView.cpp b/Engine/Plugins/Developer/Concert/ConcertApp/MultiUserServer/Source/MultiUserServer/Private/Widgets/Clients/SConcertClientsTabView.cpp
--- a/Engine/Plugins/Developer/Concert/ConcertApp/MultiUserServer/Source/MultiUserServer/Private/Widgets/Clients/SConcertClientsTabView.cpp
+++ b/Engine/Plugins/Developer/Concert/ConcertApp/MultiUserServer/Source/MultiUserServer/Private/Widgets/Clients/SConcertClientsTabView.cpp
@@ -22,235 +22,238 @@
 
 #define LOCTEXT_NAMESPACE "UnrealMultiUserUI.SConcertClientsTabView"
 
-const FName SConcertClientsTabView::ClientBrowserTabId("ClientBrowserTabId");
-const FName SConcertClientsTabView::GlobalLogTabId("GlobalLogTabId");
-const FName SConcertClientsTabView::PackageTransmissionTabId("PackageTransmissionTabId");
-
-void SConcertClientsTabView::Construct(const FArguments& InArgs, FName InStatusBarID, TSharedRef<IConcertSyncServer> InServer, TSharedRef<FGlobalLogSource> InLogBuffer)
-{
-	Server = MoveTemp(InServer);
-	LogBuffer = MoveTemp(InLogBuffer);
-	ClientInfoCache = MakeShared<FEndpointToUserNameCache>(Server->GetConcertServer());
-	LogTokenizer = MakeShared<FConcertLogTokenizer>(ClientInfoCache.ToSharedRef());
-	PackageTransmissionModel = MakeShared<UE::MultiUserServer::FPackageTransmissionModel>(Server.ToSharedRef());
-	
-	SConcertTabViewWithManagerBase::Construct(
-		SConcertTabViewWithManagerBase::FArguments()
-		.ConstructUnderWindow(InArgs._ConstructUnderWindow)
-		.ConstructUnderMajorTab(InArgs._ConstructUnderMajorTab)
-		.CreateTabs(FCreateTabs::CreateLambda([this, &InArgs](const TSharedRef<FTabManager>& InTabManager, const TSharedRef<FTabManager::FLayout>& InLayout)
-		{
-			CreateTabs(InTabManager, InLayout, InArgs);
-		}))
-		.LayoutName("ConcertClientsTabView_v0.2"),
-		InStatusBarID
-	);
-}
-
-void SConcertClientsTabView::ShowConnectedClients(const FGuid& SessionId) const
+namespace UE::MultiUserServer
 {
-	ClientBrowser->ShowOnlyClientsFromSession(SessionId);
-}
+	const FName SConcertClientsTabView::ClientBrowserTabId("ClientBrowserTabId");
+	const FName SConcertClientsTabView::GlobalLogTabId("GlobalLogTabId");
+	const FName SConcertClientsTabView::PackageTransmissionTabId("PackageTransmissionTabId");
 
-void SConcertClientsTabView::OpenClientLogTab(const FGuid& ClientMessageNodeId) const
-{
-	const FName TabId = *ClientMessageNodeId.ToString();
-	if (const TSharedPtr<SDockTab> ExistingTab = GetTabManager()->FindExistingLiveTab(FTabId(TabId)))
+	void SConcertClientsTabView::Construct(const FArguments& InArgs, FName InStatusBarID, TSharedRef<IConcertSyncServer> InServer, TSharedRef<FGlobalLogSource> InLogBuffer)
 	{
-		GetTabManager()->DrawAttention(ExistingTab.ToSharedRef());
+		Server = MoveTemp(InServer);
+		LogBuffer = MoveTemp(InLogBuffer);
+		ClientInfoCache = MakeShared<FEndpointToUserNameCache>(Server->GetConcertServer());
+		LogTokenizer = MakeShared<FConcertLogTokenizer>(ClientInfoCache.ToSharedRef());
+		PackageTransmissionModel = MakeShared<FPackageTransmissionModel>(Server.ToSharedRef());
+		
+		SConcertTabViewWithManagerBase::Construct(
+			SConcertTabViewWithManagerBase::FArguments()
+			.ConstructUnderWindow(InArgs._ConstructUnderWindow)
+			.ConstructUnderMajorTab(InArgs._ConstructUnderMajorTab)
+			.CreateTabs(FCreateTabs::CreateLambda([this, &InArgs](const TSharedRef<FTabManager>& InTabManager, const TSharedRef<FTabManager::FLayout>& InLayout)
+			{
+				CreateTabs(InTabManager, InLayout, InArgs);
+			}))
+			.LayoutName("ConcertClientsTabView_v0.2"),
+			InStatusBarID
+		);
 	}
-	else
+
+	void SConcertClientsTabView::ShowConnectedClients(const FGuid& SessionId) const
 	{
-		const TOptional<FConcertClientInfo> ClientInfo = ClientInfoCache->GetClientInfoFromNodeId(ClientMessageNodeId);
-		const TSharedRef<SDockTab> NewTab = SNew(SDockTab)
-			.Label_Lambda([this, ClientMessageNodeId]()
-			{
-				const TOptional<FConcertClientInfo> ClientInfo = ClientInfoCache->GetClientInfoFromNodeId(ClientMessageNodeId);
-				return ClientInfo
-					? FText::Format(LOCTEXT("ClientTabFmt", "{0} Log"), FText::FromString(ClientInfo->DisplayName))
-					: FText::FromString(ClientMessageNodeId.ToString(EGuidFormats::DigitsWithHyphens));
-			})
-			.ToolTipText(FText::Format(LOCTEXT("ClientTabTooltipFmt", "Logs all networked requests originating or going to client {0} (NodeId = {1})"), ClientInfo ? FText::FromString(ClientInfo->DisplayName) : FText::GetEmpty(), FText::FromString(ClientMessageNodeId.ToString())))
-			.TabRole(PanelTab)
-			[
-				SNew(SConcertTransportLog, LogBuffer.ToSharedRef(), ClientInfoCache.ToSharedRef(), LogTokenizer.ToSharedRef())
-				.Filter(UE::MultiUserServer::MakeClientLogFilter(LogTokenizer.ToSharedRef(), ClientMessageNodeId, ClientInfoCache.ToSharedRef()))
-			]; 
+		ClientBrowser->ShowOnlyClientsFromSession(SessionId);
+	}
 
-		// We need a tab to place the client tab next to
-		if (IsGlobalLogOpen())
+	void SConcertClientsTabView::OpenClientLogTab(const FGuid& ClientMessageNodeId) const
+	{
+		const FName TabId = *ClientMessageNodeId.ToString();
+		if (const TSharedPtr<SDockTab> ExistingTab = GetTabManager()->FindExistingLiveTab(FTabId(TabId)))
 		{
-			const FTabManager::FLiveTabSearch Search(GlobalLogTabId);
-			GetTabManager()->InsertNewDocumentTab(TabId, Search, NewTab);
+			GetTabManager()->DrawAttention(ExistingTab.ToSharedRef());
 		}
 		else
 		{
-			OpenGlobalLogTab();
-			
-			const FTabManager::FLiveTabSearch Search(GlobalLogTabId);
-			GetTabManager()->InsertNewDocumentTab(TabId, Search, NewTab);
+			const TOptional<FConcertClientInfo> ClientInfo = ClientInfoCache->GetClientInfoFromNodeId(ClientMessageNodeId);
+			const TSharedRef<SDockTab> NewTab = SNew(SDockTab)
+				.Label_Lambda([this, ClientMessageNodeId]()
+				{
+					const TOptional<FConcertClientInfo> ClientInfo = ClientInfoCache->GetClientInfoFromNodeId(ClientMessageNodeId);
+					return ClientInfo
+						? FText::Format(LOCTEXT("ClientTabFmt", "{0} Log"), FText::FromString(ClientInfo->DisplayName))
+						: FText::FromString(ClientMessageNodeId.ToString(EGuidFormats::DigitsWithHyphens));
+				})
+				.ToolTipText(FText::Format(LOCTEXT("ClientTabTooltipFmt", "Logs all networked requests originating or going to client {0} (NodeId = {1})"), ClientInfo ? FText::FromString(ClientInfo->DisplayName) : FText::GetEmpty(), FText::FromString(ClientMessageNodeId.ToString())))
+				.TabRole(PanelTab)
+				[
+					SNew(SConcertTransportLog, LogBuffer.ToSharedRef(), ClientInfoCache.ToSharedRef(), LogTokenizer.ToSharedRef())
+					.Filter(UE::MultiUserServer::MakeClientLogFilter(LogTokenizer.ToSharedRef(), ClientMessageNodeId, ClientInfoCache.ToSharedRef()))
+				]; 
 
-			CloseGlobalLogTab();
+			// We need a tab to place the client tab next to
+			if (IsGlobalLogOpen())
+			{
+				const FTabManager::FLiveTabSearch Search(GlobalLogTabId);
+				GetTabManager()->InsertNewDocumentTab(TabId, Search, NewTab);
+			}
+			else
+			{
+				OpenGlobalLogTab();
+				
+				const FTabManager::FLiveTabSearch Search(GlobalLogTabId);
+				GetTabManager()->InsertNewDocumentTab(TabId, Search, NewTab);
+
+				CloseGlobalLogTab();
+			}
 		}
 	}
-}
 
-bool SConcertClientsTabView::CanScrollToLog(const FGuid& MessageId, FConcertLogEntryFilterFunc FilterFunc, FText& ErrorMessage) const
-{
-	const bool bCanScrollToLog = GlobalTransportLog->CanScrollToLog(MessageId, FilterFunc);
-	if (!bCanScrollToLog)
+	bool SConcertClientsTabView::CanScrollToLog(const FGuid& MessageId, FConcertLogEntryFilterFunc FilterFunc, FText& ErrorMessage) const
 	{
-		ErrorMessage = LOCTEXT("ScrollToLog.LogNotAvailable", "The log is filtered out.");
-		return false;
-	}
+		const bool bCanScrollToLog = GlobalTransportLog->CanScrollToLog(MessageId, FilterFunc);
+		if (!bCanScrollToLog)
+		{
+			ErrorMessage = LOCTEXT("ScrollToLog.LogNotAvailable", "The log is filtered out.");
+			return false;
+		}
 
-	return true;
-}
+		return true;
+	}
 
-void SConcertClientsTabView::ScrollToLog(const FGuid& MessageId, FConcertLogEntryFilterFunc FilterFunc) const
-{
-	OpenGlobalLogTab();
-	GlobalTransportLog->ScrollToLog(MessageId, FilterFunc);
-}
+	void SConcertClientsTabView::ScrollToLog(const FGuid& MessageId, FConcertLogEntryFilterFunc FilterFunc) const
+	{
+		OpenGlobalLogTab();
+		GlobalTransportLog->ScrollToLog(MessageId, FilterFunc);
+	}
 
-void SConcertClientsTabView::OpenGlobalLogTab() const
-{
-	GetTabManager()->TryInvokeTab(GlobalLogTabId);
-}
+	void SConcertClientsTabView::OpenGlobalLogTab() const
+	{
+		GetTabManager()->TryInvokeTab(GlobalLogTabId);
+	}
 
-void SConcertClientsTabView::CloseGlobalLogTab() const
-{
-	if (const TSharedPtr<SDockTab> GlobalLogTab = GetGlobalLogTab())
+	void SConcertClientsTabView::CloseGlobalLogTab() const
 	{
-		GlobalLogTab->RequestCloseTab();
+		if (const TSharedPtr<SDockTab> GlobalLogTab = GetGlobalLogTab())
+		{
+			GlobalLogTab->RequestCloseTab();
+		}
 	}
-}
 
-bool SConcertClientsTabView::IsGlobalLogOpen() const
-{
-	return GetGlobalLogTab().IsValid();
-}
+	bool SConcertClientsTabView::IsGlobalLogOpen() const
+	{
+		return GetGlobalLogTab().IsValid();
+	}
 
-TSharedPtr<SDockTab> SConcertClientsTabView::GetGlobalLogTab() const
-{
-	return GetTabManager()->FindExistingLiveTab(GlobalLogTabId);
-}
+	TSharedPtr<SDockTab> SConcertClientsTabView::GetGlobalLogTab() const
+	{
+		return GetTabManager()->FindExistingLiveTab(GlobalLogTabId);
+	}
 
-void SConcertClientsTabView::CreateTabs(const TSharedRef<FTabManager>& InTabManager, const TSharedRef<FTabManager::FLayout>& InLayout, const FArguments& InArgs)
-{
-	MainPackageTransmissionTab = MakeShared<UE::MultiUserServer::FPackageTransmissionTabController>(
-		PackageTransmissionTabId,
-		InTabManager,
-		PackageTransmissionModel.ToSharedRef(),
-		ClientInfoCache.ToSharedRef(),
-		UE::MultiUserServer::FCanScrollToLog::CreateSP(this, &SConcertClientsTabView::CanScrollToLog),
-		UE::MultiUserServer::FScrollToLog::CreateSP(this, &SConcertClientsTabView::ScrollToLog)
-		);
-	
-	InTabManager->RegisterTabSpawner(ClientBrowserTabId, FOnSpawnTab::CreateSP(this, &SConcertClientsTabView::SpawnClientBrowserTab))
-		.SetDisplayName(LOCTEXT("ClientBrowserTabLabel", "Clients"));
-	InTabManager->RegisterTabSpawner(GlobalLogTabId, FOnSpawnTab::CreateSP(this, &SConcertClientsTabView::SpawnGlobalLogTab))
-		.SetDisplayName(LOCTEXT("GlobalLogTabLabel", "Global Log"));
-	InLayout->AddArea
-		(
-			FTabManager::NewPrimaryArea()
-				->SetOrientation(Orient_Vertical)
-				->Split
-				(
-					FTabManager::NewSplitter()
-					->SetSizeCoefficient(0.5f)
-					->SetOrientation(Orient_Horizontal)
+	void SConcertClientsTabView::CreateTabs(const TSharedRef<FTabManager>& InTabManager, const TSharedRef<FTabManager::FLayout>& InLayout, const FArguments& InArgs)
+	{
+		MainPackageTransmissionTab = MakeShared<UE::MultiUserServer::FPackageTransmissionTabController>(
+			PackageTransmissionTabId,
+			InTabManager,
+			PackageTransmissionModel.ToSharedRef(),
+			ClientInfoCache.ToSharedRef(),
+			UE::MultiUserServer::FCanScrollToLog::CreateSP(this, &SConcertClientsTabView::CanScrollToLog),
+			UE::MultiUserServer::FScrollToLog::CreateSP(this, &SConcertClientsTabView::ScrollToLog)
+			);
+		
+		InTabManager->RegisterTabSpawner(ClientBrowserTabId, FOnSpawnTab::CreateSP(this, &SConcertClientsTabView::SpawnClientBrowserTab))
+			.SetDisplayName(LOCTEXT("ClientBrowserTabLabel", "Clients"));
+		InTabManager->RegisterTabSpawner(GlobalLogTabId, FOnSpawnTab::CreateSP(this, &SConcertClientsTabView::SpawnGlobalLogTab))
+			.SetDisplayName(LOCTEXT("GlobalLogTabLabel", "Global Log"));
+		InLayout->AddArea
+			(
+				FTabManager::NewPrimaryArea()
+					->SetOrientation(Orient_Vertical)
 					->Split
 					(
-						FTabManager::NewStack()
-						->SetSizeCoefficient(0.65f)
-						->AddTab(ClientBrowserTabId, ETabState::OpenedTab)
+						FTabManager::NewSplitter()
+						->SetSizeCoefficient(0.5f)
+						->SetOrientation(Orient_Horizontal)
+						->Split
+						(
+							FTabManager::NewStack()
+							->SetSizeCoefficient(0.65f)
+							->AddTab(ClientBrowserTabId, ETabState::OpenedTab)
+						)
+						->Split
+						(
+							FTabManager::NewStack()
+							->SetSizeCoefficient(0.35f)
+							->AddTab(PackageTransmissionTabId, ETabState::OpenedTab)
+								
+						)
 					)
 					->Split
 					(
-						FTabManager::NewStack()
-						->SetSizeCoefficient(0.35f)
-						->AddTab(PackageTransmissionTabId, ETabState::OpenedTab)
-							
+					FTabManager::NewStack()
+						->SetSizeCoefficient(0.6f)
+						->AddTab(GlobalLogTabId, ETabState::OpenedTab)
 					)
-				)
-				->Split
-				(
-				FTabManager::NewStack()
-					->SetSizeCoefficient(0.6f)
-					->AddTab(GlobalLogTabId, ETabState::OpenedTab)
-				)
-		);
-}
+			);
+	}
 
-TSharedRef<SDockTab> SConcertClientsTabView::SpawnClientBrowserTab(const FSpawnTabArgs& InTabArgs)
-{
-	using namespace UE::MultiUserServer;
-	const TSharedRef<FClientNetworkStatisticsModel> NetworkStatisticsModel = MakeShared<FClientNetworkStatisticsModel>();
-	
-	return SNew(SDockTab)
-		.Label(LOCTEXT("ClientBrowserTabLabel", "Clients"))
-		.TabRole(PanelTab)
-		[
-			SAssignNew(ClientBrowser, UE::MultiUserServer::SConcertNetworkBrowser,
-				MakeShared<UE::MultiUserServer::FClientBrowserModel>(Server->GetConcertServer(), ClientInfoCache.ToSharedRef(), NetworkStatisticsModel), NetworkStatisticsModel, Server->GetConcertServer())
-			.RightOfSearch()
+	TSharedRef<SDockTab> SConcertClientsTabView::SpawnClientBrowserTab(const FSpawnTabArgs& InTabArgs)
+	{
+		using namespace UE::MultiUserServer;
+		const TSharedRef<FClientNetworkStatisticsModel> NetworkStatisticsModel = MakeShared<FClientNetworkStatisticsModel>();
+		
+		return SNew(SDockTab)
+			.Label(LOCTEXT("ClientBrowserTabLabel", "Clients"))
+			.TabRole(PanelTab)
 			[
-				CreateOpenGlobalLogButton()
-			]
-			.OnServerDoubleClicked_Raw(this, &SConcertClientsTabView::OpenGlobalLogTab)
-			.OnClientDoubleClicked_Raw(this, &SConcertClientsTabView::OpenClientLogTab)
-		]; 
-}
+				SAssignNew(ClientBrowser, UE::MultiUserServer::SConcertNetworkBrowser,
+					MakeShared<UE::MultiUserServer::FClientBrowserModel>(Server->GetConcertServer(), ClientInfoCache.ToSharedRef(), NetworkStatisticsModel), NetworkStatisticsModel, Server->GetConcertServer())
+				.RightOfSearch()
+				[
+					CreateOpenGlobalLogButton()
+				]
+				.OnServerDoubleClicked_Raw(this, &SConcertClientsTabView::OpenGlobalLogTab)
+				.OnClientDoubleClicked_Raw(this, &SConcertClientsTabView::OpenClientLogTab)
+			]; 
+	}
 
-TSharedRef<SDockTab> SConcertClientsTabView::SpawnGlobalLogTab(const FSpawnTabArgs& InTabArgs)
-{
-	return SNew(SDockTab)
-		.Label(LOCTEXT("GlobalLogTabLabel", "Global Log"))
-		.TabRole(PanelTab)
-		[
-			SAssignNew(GlobalTransportLog, SConcertTransportLog, LogBuffer.ToSharedRef(), ClientInfoCache.ToSharedRef(), LogTokenizer.ToSharedRef())
-			.Filter(UE::MultiUserServer::MakeGlobalLogFilter(LogTokenizer.ToSharedRef()))
-		]; 
-}
+	TSharedRef<SDockTab> SConcertClientsTabView::SpawnGlobalLogTab(const FSpawnTabArgs& InTabArgs)
+	{
+		return SNew(SDockTab)
+			.Label(LOCTEXT("GlobalLogTabLabel", "Global Log"))
+			.TabRole(PanelTab)
+			[
+				SAssignNew(GlobalTransportLog, SConcertTransportLog, LogBuffer.ToSharedRef(), ClientInfoCache.ToSharedRef(), LogTokenizer.ToSharedRef())
+				.Filter(UE::MultiUserServer::MakeGlobalLogFilter(LogTokenizer.ToSharedRef()))
+			]; 
+	}
 
-TSharedRef<SWidget> SConcertClientsTabView::CreateOpenGlobalLogButton() const
-{
-	return SNew(SButton)
-		.ButtonStyle(FAppStyle::Get(), "SimpleButton")
-		.ToolTipText(LOCTEXT("OpenGlobalLogTooltip", "Opens the Global Log which logs all incoming networked messages."))
-		.ContentPadding(FMargin(1, 0))
-		.Visibility_Lambda(
-			[this]()
+	TSharedRef<SWidget> SConcertClientsTabView::CreateOpenGlobalLogButton() const
+	{
+		return SNew(SButton)
+			.ButtonStyle(FAppStyle::Get(), "SimpleButton")
+			.ToolTipText(LOCTEXT("OpenGlobalLogTooltip", "Opens the Global Log which logs all incoming networked messages."))
+			.ContentPadding(FMargin(1, 0))
+			.Visibility_Lambda(
+				[this]()
+				{
+					const bool bIsGlobalLogOpen = IsGlobalLogOpen();
+					return bIsGlobalLogOpen ? EVisibility::Hidden : EVisibility::Visible;
+				})
+			.OnClicked_Lambda([this]()
 			{
-				const bool bIsGlobalLogOpen = IsGlobalLogOpen();
-				return bIsGlobalLogOpen ? EVisibility::Hidden : EVisibility::Visible;
+				OpenGlobalLogTab();
+				return FReply::Handled();
 			})
-		.OnClicked_Lambda([this]()
-		{
-			OpenGlobalLogTab();
-			return FReply::Handled();
-		})
-		[
-			SNew(SHorizontalBox)
-				+SHorizontalBox::Slot()
-				.AutoWidth()
-				.VAlign(VAlign_Center)
-				.Padding(4.0, 0.0f)
-				[
-					SNew(SImage)
-					.ColorAndOpacity(FSlateColor::UseForeground())
-					.Image(FAppStyle::Get().GetBrush("Icons.Layout"))
-				]
-				+SHorizontalBox::Slot()
-				.VAlign(VAlign_Center)
-				.Padding(4.0, 0.0f)
-				[
-					SNew(STextBlock)
-					.Text(LOCTEXT("OpenGlobalLog", "Open Global Log"))
-					.ColorAndOpacity(FSlateColor::UseForeground())
-				]
-		];
+			[
+				SNew(SHorizontalBox)
+					+SHorizontalBox::Slot()
+					.AutoWidth()
+					.VAlign(VAlign_Center)
+					.Padding(4.0, 0.0f)
+					[
+						SNew(SImage)
+						.ColorAndOpacity(FSlateColor::UseForeground())
+						.Image(FAppStyle::Get().GetBrush("Icons.Layout"))
+					]
+					+SHorizontalBox::Slot()
+					.VAlign(VAlign_Center)
+					.Padding(4.0, 0.0f)
+					[
+						SNew(STextBlock)
+						.Text(LOCTEXT("OpenGlobalLog", "Open Global Log"))
+						.ColorAndOpacity(FSlateColor::UseForeground())
+					]
+			];
+	}
 }
 
 #undef LOCTEXT_NAMESPACE
diff --git a/Engine/Plugins/Developer/Concert/ConcertApp/MultiUserServer/Source/MultiUserServer/Private/Widgets/Clients/SConcertClientsTabView.h b/Engine/Plugins/Developer/Concert/ConcertApp/MultiUserServer/Source/MultiUserServer/Private/Widgets/Clients/SConcertClientsTabView.h
--- a/Engine/Plugins/Developer/Concert/ConcertApp/MultiUserServer/Source/MultiUserServer/Private/Widgets/Clients/SConcertClientsTabView.h
+++ b/Engine/Plugins/Developer/Concert/ConcertApp/MultiUserServer/Source/MultiUserServer/Private/Widgets/Clients/SConcertClientsTabView.h
@@ -7,20 +7,7 @@
 #include "Widgets/DeclarativeSyntaxSupport.h"
 #include "Widgets/SConcertTabViewWithManagerBase.h"
 
-class SConcertTransportLog;
-
-namespace UE
-{
-	namespace MultiUserServer
-	{
-		class FPackageTransmissionTabController;
-		class FPackageTransmissionModel;
-		class SConcertNetworkBrowser;
-	}
-}
-
 class FConcertClientsTabController;
-class FConcertLogTokenizer;
 class FEndpointToUserNameCache;
 class FGlobalLogSource;
 class FTabManager;
@@ -31,60 +18,69 @@ class SWidget;
 class SWindow;
 struct FConcertClientInfo;
 
-/** Manages the UI logic of the Clients tab */
-class SConcertClientsTabView : public SConcertTabViewWithManagerBase
+namespace UE::MultiUserServer
 {
-public:
+	class FConcertLogTokenizer;
+	class FPackageTransmissionTabController;
+	class FPackageTransmissionModel;
+	class SConcertNetworkBrowser;
+	class SConcertTransportLog;
 	
-	static const FName ClientBrowserTabId;
-	static const FName GlobalLogTabId;
-	static const FName PackageTransmissionTabId;
+	/** Manages the UI logic of the Clients tab */
+	class SConcertClientsTabView : public SConcertTabViewWithManagerBase
+	{
+	public:
+		
+		static const FName ClientBrowserTabId;
+		static const FName GlobalLogTabId;
+		static const FName PackageTransmissionTabId;
 
-	SLATE_BEGIN_ARGS(SConcertClientsTabView)
-	{}
-		SLATE_ARGUMENT(TSharedPtr<SDockTab>, ConstructUnderMajorTab)
-		SLATE_ARGUMENT(TSharedPtr<SWindow>, ConstructUnderWindow)
-	SLATE_END_ARGS()
+		SLATE_BEGIN_ARGS(SConcertClientsTabView)
+		{}
+			SLATE_ARGUMENT(TSharedPtr<SDockTab>, ConstructUnderMajorTab)
+			SLATE_ARGUMENT(TSharedPtr<SWindow>, ConstructUnderWindow)
+		SLATE_END_ARGS()
 
-	void Construct(const FArguments& InArgs, FName InStatusBarID, TSharedRef<IConcertSyncServer> InServer, TSharedRef<FGlobalLogSource> InLogBuffer);
-	
-	void ShowConnectedClients(const FGuid& SessionId) const;
-	void OpenGlobalLogTab() const;
-	void CloseGlobalLogTab() const;
-	void OpenClientLogTab(const FGuid& ClientMessageNodeId) const;
+		void Construct(const FArguments& InArgs, FName InStatusBarID, TSharedRef<IConcertSyncServer> InServer, TSharedRef<FGlobalLogSource> InLogBuffer);
+		
+		void ShowConnectedClients(const FGuid& SessionId) const;
+		void OpenGlobalLogTab() const;
+		void CloseGlobalLogTab() const;
+		void OpenClientLogTab(const FGuid& ClientMessageNodeId) const;
 
-	/** Whether it is possible to call ScrollToLog. */
-	bool CanScrollToLog(const FGuid& MessageId, FConcertLogEntryFilterFunc FilterFunc, FText& ErrorMessage) const;
-	/** Scrolls to the log in the global log tab. */
-	void ScrollToLog(const FGuid& MessageId, FConcertLogEntryFilterFunc FilterFunc) const;
-	
-	bool IsGlobalLogOpen() const;
-	TSharedPtr<SDockTab> GetGlobalLogTab() const;
-	
-private:
+		/** Whether it is possible to call ScrollToLog. */
+		bool CanScrollToLog(const FGuid& MessageId, FConcertLogEntryFilterFunc FilterFunc, FText& ErrorMessage) const;
+		/** Scrolls to the log in the global log tab. */
+		void ScrollToLog(const FGuid& MessageId, FConcertLogEntryFilterFunc FilterFunc) const;
+		
+		bool IsGlobalLogOpen() const;
+		TSharedPtr<SDockTab> GetGlobalLogTab() const;
+		
+	private:
 
-	/** Used to look up client info */
-	TSharedPtr<IConcertSyncServer> Server;
-	/** Buffers all logs globally */
-	TSharedPtr<FGlobalLogSource> LogBuffer;
+		/** Used to look up client info */
+		TSharedPtr<IConcertSyncServer> Server;
+		/** Buffers all logs globally */
+		TSharedPtr<FGlobalLogSource> LogBuffer;
 
-	/** Caches client info so it remains available even after a client disconnects */
-	TSharedPtr<FEndpointToUserNameCache> ClientInfoCache;
-	/** Used by various systems to convert logs to text */
-	TSharedPtr<FConcertLogTokenizer> LogTokenizer;
-	
-	/** Knows about when packages are sent and received by the server */
-	TSharedPtr<UE::MultiUserServer::FPackageTransmissionModel> PackageTransmissionModel;
-	TSharedPtr<UE::MultiUserServer::FPackageTransmissionTabController> MainPackageTransmissionTab;
+		/** Caches client info so it remains available even after a client disconnects */
+		TSharedPtr<FEndpointToUserNameCache> ClientInfoCache;
+		/** Used by various systems to convert logs to text */
+		TSharedPtr<FConcertLogTokenizer> LogTokenizer;
+		
+		/** Knows about when packages are sent and received by the server */
+		TSharedPtr<FPackageTransmissionModel> PackageTransmissionModel;
+		TSharedPtr<FPackageTransmissionTabController> MainPackageTransmissionTab;
 
-	TSharedPtr<UE::MultiUserServer::SConcertNetworkBrowser> ClientBrowser;
+		TSharedPtr<SConcertNetworkBrowser> ClientBrowser;
 
-	/** The widget inside of the global log tab */
-	TSharedPtr<SConcertTransportLog> GlobalTransportLog;
-	
-	void CreateTabs(const TSharedRef<FTabManager>& InTabManager, const TSharedRef<FTabManager::FLayout>& InLayout, const FArguments& InArgs);
-	TSharedRef<SDockTab> SpawnClientBrowserTab(const FSpawnTabArgs& InTabArgs);
-	TSharedRef<SDockTab> SpawnGlobalLogTab(const FSpawnTabArgs& InTabArgs);
+		/** The widget inside of the global log tab */
+		TSharedPtr<SConcertTransportLog> GlobalTransportLog;
+		
+		void CreateTabs(const TSharedRef<FTabManager>& InTabManager, const TSharedRef<FTabManager::FLayout>& InLayout, const FArguments& InArgs);
+		TSharedRef<SDockTab> SpawnClientBrowserTab(const FSpawnTabArgs& InTabArgs);
+		TSharedRef<SDockTab> SpawnGlobalLogTab(const FSpawnTabArgs& InTabArgs);
 
-	TSharedRef<SWidget> CreateOpenGlobalLogButton() const;
-};
+		TSharedRef<SWidget> CreateOpenGlobalLogButton() const;
+	};
+}
diff --git a/Engine/Plugins/Developer/Concert/ConcertApp/MultiUserServer/Source/MultiUserServer/Private/Widgets/Util/Filter/ConcertFilter.h b/Engine/Plugins/Developer/Concert/ConcertApp/MultiUserServer/Source/MultiUserServer/Private/Widgets/Util/Filter/ConcertFilter.h
new file mode 100644
--- /dev/null
+++ b/Engine/Plugins/Developer/Concert/ConcertApp/MultiUserServer/Source/MultiUserServer/Private/Widgets/Util/Filter/ConcertFilter.h
@@ -0,0 +1,35 @@
+// Copyright Epic Games, Inc. All Rights Reserved.
+
+#pragma once
+
+#include "CoreMinimal.h"
+#include "Misc/IFilter.h"
+#include "Templates/UnrealTemplate.h"
+
+namespace UE::MultiUserServer
+{
+	/** Base filter */
+	template<typename TFilterType>
+	class TConcertFilter
+		:
+		public IFilter<TFilterType>,
+		// We do not need to copy filters so let's avoid it by accident; some constructors pass this pointer to callbacks which become stale upon copying
+		public FNoncopyable
+	{
+	public:
+
+		//~ Begin IFilter Interface
+		DECLARE_DERIVED_EVENT( FFrontendFilter, IFilter<TFilterType>::template FChangedEvent, FChangedEvent );
+		virtual FChangedEvent& OnChanged() override { return ChangedEvent; }
+		//~ End IFilter Interface
+
+	protected:
+	
+		void BroadcastChangedEvent() const { ChangedEvent.Broadcast(); }
+
+	private:
+	
+		FChangedEvent ChangedEvent;
+	};
+}
+
diff --git a/Engine/Plugins/Developer/Concert/ConcertApp/MultiUserServer/Source/MultiUserServer/Private/Widgets/Util/Filter/ConcertFrontendFilter.h b/Engine/Plugins/Developer/Concert/ConcertApp/MultiUserServer/Source/MultiUserServer/Private/Widgets/Util/Filter/ConcertFrontendFilter.h
new file mode 100644
--- /dev/null
+++ b/Engine/Plugins/Developer/Concert/ConcertApp/MultiUserServer/Source/MultiUserServer/Private/Widgets/Util/Filter/ConcertFrontendFilter.h
@@ -0,0 +1,63 @@
+// Copyright Epic Games, Inc. All Rights Reserved.
+
+#pragma once
+
+#include "CoreMinimal.h"
+#include "ConcertFilter.h"
+
+class SWidget;
+
+namespace UE::MultiUserServer
+{
+	/** A filter that is intended to be displayed in the UI. Every filter has one widget displaying it. */
+	template<typename TFilterType>
+	class TConcertFrontendFilter
+		:
+		public TConcertFilter<TFilterType>,
+		public TSharedFromThis<TConcertFrontendFilter<TFilterType>>
+	{
+	public:
+
+		/** Gets a widget that represents this filter */
+		virtual TSharedRef<SWidget> GetFilterWidget() = 0;
+	};
+
+	/**
+	 * Helper class to implement filters
+	 *
+	 * Intended pattern:
+	 *	1. Subclass TConcertFilter<TFilterType> and implement filter logic in TConcertFilterImpl, e.g. text search. This will act as a "model" in MVC.
+	 *  2. Subclass TConcertFrontendLogFilterAggregate and handle creating UI in it. This will act as a "View" in MVC.
+	 */
+	template<typename TConcertFilterImpl, typename TFilterType, typename TWidgetType = SWidget>
+	class TConcertFrontendFilterAggregate : public TConcertFrontendFilter<TFilterType>
+	{
+		using Super = TConcertFrontendFilter<TFilterType>;
+	public:
+
+		//~ Begin FConcertLogFilter Interface
+		virtual bool PassesFilter(TFilterType InItem) const final override { return Implementation.PassesFilter(InItem); }
+		//~ End FConcertLogFilter Interface
+	
+		//~ Begin FConcertFrontendLogFilter Interface
+		virtual TSharedRef<SWidget> GetFilterWidget() final override { return ChildSlot.ToSharedRef(); }
+		//~ End FConcertLogFilter Interface
+	
+	protected:
+
+		template<typename... TArg>
+		TConcertFrontendFilterAggregate(TArg&&... Arg)
+			: Implementation(Forward<TArg>(Arg)...)
+		{
+			Implementation.OnChanged().AddLambda([this]()
+			{
+				Super::template OnChanged().Broadcast();
+			});
+		}
+
+		TConcertFilterImpl Implementation;
+		TSharedPtr<TWidgetType> ChildSlot;
+	};
+}
+
+
diff --git a/Engine/Plugins/Developer/Concert/ConcertApp/MultiUserServer/Source/MultiUserServer/Private/Widgets/Util/Filter/ConcertFrontendFilter_TextSearch.h b/Engine/Plugins/Developer/Concert/ConcertApp/MultiUserServer/Source/MultiUserServer/Private/Widgets/Util/Filter/ConcertFrontendFilter_TextSearch.h
new file mode 100644
--- /dev/null
+++ b/Engine/Plugins/Developer/Concert/ConcertApp/MultiUserServer/Source/MultiUserServer/Private/Widgets/Util/Filter/ConcertFrontendFilter_TextSearch.h
@@ -0,0 +1,76 @@
+// Copyright Epic Games, Inc. All Rights Reserved.
+
+#pragma once
+
+#include "CoreMinimal.h"
+#include "ConcertFrontendFilter.h"
+#include "Misc/TextFilter.h"
+#include "Widgets/Input/SSearchBox.h"
+
+class FConcertLogTokenizer;
+
+namespace UE::MultiUserServer
+{
+	/** Allows advanced search by text. Implements Adapter pattern to wrap TTextFilter. */
+	template<typename TFilterType>
+	class TConcertFilter_TextSearch : public TConcertFilter<TFilterType>
+	{
+		using Super = TConcertFilter<TFilterType>;
+	public:
+
+		TConcertFilter_TextSearch()
+			: TextFilter(TTextFilter<TFilterType>::template FItemToStringArray::CreateRaw(this, &TConcertFilter_TextSearch<TFilterType>::GenerateSearchTerms))
+		{
+			TextFilter.OnChanged().AddLambda([this]
+			{
+				Super::template OnChanged().Broadcast();
+			});
+		}
+	
+		//~ Begin FConcertLogFilter Interface
+		virtual bool PassesFilter(TFilterType InItem) const override { return TextFilter.PassesFilter(InItem); }
+		//~ End FConcertLogFilter Interface
+	
+		void SetRawFilterText(const FText& InFilterText) { TextFilter.SetRawFilterText(InFilterText); }
+
+	protected:
+		
+		/** Parses InItem into a bunch of strings that can be searched */
+		virtual void GenerateSearchTerms(TFilterType InItem, TArray<FString>& OutTerms) const = 0;
+		
+	private:
+
+		/** Does the actual string search */
+		TTextFilter<TFilterType> TextFilter;
+	};
+
+	/** Creates a search bar */
+	template<typename TTextSearchFilterType /* Expected subtype of TConcertFilter_TextSearch */, typename TFilterType>
+	class TConcertFrontendFilter_TextSearch : public TConcertFrontendFilterAggregate<TTextSearchFilterType, TFilterType, SSearchBox>
+	{
+		using Super = TConcertFrontendFilterAggregate<TTextSearchFilterType, TFilterType, SSearchBox>;
+	public:
+
+		template<typename... TArg>
+		TConcertFrontendFilter_TextSearch(TArg&&... Arg)
+			: Super(Forward<TArg>(Arg)...)
+		{
+			Super::template ChildSlot = SNew(SSearchBox)
+				.OnTextChanged_Lambda([this](const FText& NewSearchText)
+				{
+					SearchText = NewSearchText;
+					TTextSearchFilterType& Impl = Super::template Implementation;
+					Impl.SetRawFilterText(NewSearchText);
+				})
+				.DelayChangeNotificationsWhileTyping(true);
+		}
+		
+		const FText& GetSearchText() const { return SearchText; }
+
+	private:
+
+		FText SearchText;
+	};
+}
+
+
diff --git a/Engine/Plugins/Developer/Concert/ConcertApp/MultiUserServer/Source/MultiUserServer/Private/Widgets/Util/Filter/ConcertRootFrontendFilter.h b/Engine/Plugins/Developer/Concert/ConcertApp/MultiUserServer/Source/MultiUserServer/Private/Widgets/Util/Filter/ConcertRootFrontendFilter.h
new file mode 100644
--- /dev/null
+++ b/Engine/Plugins/Developer/Concert/ConcertApp/MultiUserServer/Source/MultiUserServer/Private/Widgets/Util/Filter/ConcertRootFrontendFilter.h
@@ -0,0 +1,123 @@
+// Copyright Epic Games, Inc. All Rights Reserved.
+
+#pragma once
+
+#include "CoreMinimal.h"
+#include "ConcertFilter.h"
+#include "ConcertFrontendFilter.h"
+#include "ConcertFrontendFilter_TextSearch.h"
+
+#include "Algo/AllOf.h"
+#include "Widgets/SBoxPanel.h"
+
+namespace UE::MultiUserServer
+{
+	/** A filter that contains multiple UI filters */
+	template<typename TFilterType, typename TTextSearchFilterType>
+	class TConcertFrontendRootFilter 
+		:
+		public TConcertFilter<TFilterType>,
+		public TSharedFromThis<TConcertFrontendRootFilter<TFilterType, TTextSearchFilterType>>
+	{
+		using Super = TConcertFilter<TFilterType>;
+	public:
+
+		using FConcertFilter = TConcertFilter<TFilterType>;
+		using FConcertFrontendFilter = TConcertFrontendFilter<TFilterType>;
+		using FConcertFrontendFilter_TextSearch = TTextSearchFilterType;
+
+		TConcertFrontendRootFilter(
+			TSharedRef<FConcertFrontendFilter_TextSearch> TextSearchFilter,
+			TArray<TSharedRef<FConcertFrontendFilter>> FrontendFilters,
+			const TArray<TSharedRef<FConcertFilter>>& NonVisualFilters = {}
+			)
+			: TextSearchFilter(MoveTemp(TextSearchFilter))
+			, FrontendFilters(MoveTemp(FrontendFilters))
+			, AllFilters(this->FrontendFilters)
+		{
+			AllFilters.Reserve(FrontendFilters.Num() + 1 + NonVisualFilters.Num());
+			AllFilters.Add(TextSearchFilter);
+			for (const TSharedRef<FConcertFilter>& NonVisualFilter : NonVisualFilters)
+			{
+				AllFilters.Add(NonVisualFilter);
+			}
+	
+			for (const TSharedRef<FConcertFilter>& Filter : AllFilters)
+			{
+				Filter->OnChanged().AddRaw(this, &TConcertFrontendRootFilter<TFilterType, TTextSearchFilterType>::BroadcastOnChanged);
+			}
+		}
+
+		/** Builds the widget view for all contained filters */
+		TSharedRef<SWidget> BuildFilterWidgets() const
+		{
+			return SNew(SVerticalBox)
+
+				// Search bar
+				+SVerticalBox::Slot()
+				.AutoHeight()
+				[
+					TextSearchFilter->GetFilterWidget()
+				]
+
+				+SVerticalBox::Slot()
+				.AutoHeight()
+				.Padding(0, 2)
+				[
+					BuildCustomFilterListWidget()
+				];
+		}
+	
+		//~ Begin IFilter Interface
+		virtual bool PassesFilter(TFilterType InItem) const override
+		{
+			return Algo::AllOf(
+				AllFilters,
+				[&InItem](const TSharedRef<FConcertFilter>& AndFilter){ return AndFilter->PassesFilter(InItem); }
+				);
+		}
+		//~ End IFilter Interface
+
+		FORCEINLINE const TSharedRef<FConcertFrontendFilter_TextSearch>& GetTextSearchFilter() const { return TextSearchFilter; }
+
+	private:
+
+		/** The text search filter. Also in AllFilters. Separate variable to build search bar in new line. */
+		TSharedRef<FConcertFrontendFilter_TextSearch> TextSearchFilter;
+
+		/** AllFrontendFilters without TextSearchFilter. */
+		TArray<TSharedRef<FConcertFrontendFilter>> FrontendFilters;
+	
+		/** Filters that are combined using logical AND. */
+		TArray<TSharedRef<FConcertFilter>> AllFilters;
+
+		void BroadcastOnChanged()
+		{
+			Super::template OnChanged().Broadcast();
+		}
+
+		/** Builds the widgets that go under the text */
+		TSharedRef<SWidget> BuildCustomFilterListWidget() const
+		{
+			TSharedRef<SHorizontalBox> Box = SNew(SHorizontalBox);
+
+			bool bIsFirst = true;
+			for (const TSharedRef<FConcertFrontendFilter>& Filter : FrontendFilters)
+			{
+				const FMargin Margin = bIsFirst ? FMargin() : FMargin(8, 0, 0, 0);
+				bIsFirst = false;
+				Box->AddSlot()
+					.AutoWidth()
+					.Padding(Margin)
+					.VAlign(VAlign_Center)
+					[
+						Filter->GetFilterWidget()
+					];
+			}
+
+			return Box;
+		}
+	};
+}
+
+
diff --git a/Engine/Plugins/Developer/Concert/ConcertApp/MultiUserServer/Source/MultiUserServer/Private/Window/ConcertServerWindowController.cpp b/Engine/Plugins/Developer/Concert/ConcertApp/MultiUserServer/Source/MultiUserServer/Private/Window/ConcertServerWindowController.cpp
--- a/Engine/Plugins/Developer/Concert/ConcertApp/MultiUserServer/Source/MultiUserServer/Private/Window/ConcertServerWindowController.cpp
+++ b/Engine/Plugins/Developer/Concert/ConcertApp/MultiUserServer/Source/MultiUserServer/Private/Window/ConcertServerWindowController.cpp
@@ -18,164 +18,166 @@
 #include "Misc/ConfigCacheIni.h"
 #include "OutputLog/Public/OutputLogModule.h"
 
-
 #include "Widgets/SWindow.h"
 
 #define LOCTEXT_NAMESPACE "UnrealMultiUserUI"
 
-FConcertServerWindowController::FConcertServerWindowController(const FConcertServerWindowInitParams& Params)
-	: MultiUserServerLayoutIni(Params.MultiUserServerLayoutIni)
-	, ServerInstance(Params.Server)
-	, SessionBrowserController(MakeShared<FConcertServerSessionBrowserController>())
-	, ClientsController(MakeShared<FConcertClientsTabController>())
-	, ConcertComponents(Params.AdditionalConcertComponents)
+namespace UE::MultiUserServer
 {
-	ConcertComponents.Add(SessionBrowserController);
-	ConcertComponents.Add(ClientsController);
-}
+	FConcertServerWindowController::FConcertServerWindowController(const FConcertServerWindowInitParams& Params)
+		: MultiUserServerLayoutIni(Params.MultiUserServerLayoutIni)
+		, ServerInstance(Params.Server)
+		, SessionBrowserController(MakeShared<FConcertServerSessionBrowserController>())
+		, ClientsController(MakeShared<FConcertClientsTabController>())
+		, ConcertComponents(Params.AdditionalConcertComponents)
+	{
+		ConcertComponents.Add(SessionBrowserController);
+		ConcertComponents.Add(ClientsController);
+	}
 
-FConcertServerWindowController::~FConcertServerWindowController()
-{
-	UnregisterFromSessionDestructionEvents();
-}
+	FConcertServerWindowController::~FConcertServerWindowController()
+	{
+		UnregisterFromSessionDestructionEvents();
+	}
 
-TSharedRef<SWindow> FConcertServerWindowController::CreateWindow()
-{
-	FDisplayMetrics DisplayMetrics;
-	FSlateApplication::Get().GetDisplayMetrics(DisplayMetrics);
-	const float DPIScaleFactor = FPlatformApplicationMisc::GetDPIScaleFactorAtPoint(DisplayMetrics.PrimaryDisplayWorkAreaRect.Left, DisplayMetrics.PrimaryDisplayWorkAreaRect.Top);
-
-	const bool bEmbedTitleAreaContent = false;
-	const FVector2D ClientSize(1000.0f * DPIScaleFactor, 800.0f * DPIScaleFactor);
-	TSharedRef<SWindow> RootWindowRef = SNew(SWindow)
-		.Title(LOCTEXT("WindowTitle", "Unreal Multi User Server"))
-		.CreateTitleBar(!bEmbedTitleAreaContent)
-		.SupportsMaximize(true)
-		.SupportsMinimize(true)
-		.IsInitiallyMaximized(false)
-		.IsInitiallyMinimized(false)
-		.SizingRule(ESizingRule::UserSized)
-		.AutoCenter(EAutoCenter::PreferredWorkArea)
-		.ClientSize(ClientSize)
-		.AdjustInitialSizeAndPositionForDPIScale(false);
-	RootWindow = RootWindowRef;
+	TSharedRef<SWindow> FConcertServerWindowController::CreateWindow()
+	{
+		FDisplayMetrics DisplayMetrics;
+		FSlateApplication::Get().GetDisplayMetrics(DisplayMetrics);
+		const float DPIScaleFactor = FPlatformApplicationMisc::GetDPIScaleFactorAtPoint(DisplayMetrics.PrimaryDisplayWorkAreaRect.Left, DisplayMetrics.PrimaryDisplayWorkAreaRect.Top);
+
+		const bool bEmbedTitleAreaContent = false;
+		const FVector2D ClientSize(1000.0f * DPIScaleFactor, 800.0f * DPIScaleFactor);
+		TSharedRef<SWindow> RootWindowRef = SNew(SWindow)
+			.Title(LOCTEXT("WindowTitle", "Unreal Multi User Server"))
+			.CreateTitleBar(!bEmbedTitleAreaContent)
+			.SupportsMaximize(true)
+			.SupportsMinimize(true)
+			.IsInitiallyMaximized(false)
+			.IsInitiallyMinimized(false)
+			.SizingRule(ESizingRule::UserSized)
+			.AutoCenter(EAutoCenter::PreferredWorkArea)
+			.ClientSize(ClientSize)
+			.AdjustInitialSizeAndPositionForDPIScale(false);
+		RootWindow = RootWindowRef;
+			
+		const bool bShowRootWindowImmediately = false;
+		FSlateApplication::Get().AddWindow(RootWindowRef, bShowRootWindowImmediately);
+		FGlobalTabmanager::Get()->SetRootWindow(RootWindowRef);
+		FGlobalTabmanager::Get()->SetAllowWindowMenuBar(true);
+		FSlateNotificationManager::Get().SetRootWindow(RootWindowRef);
 		
-	const bool bShowRootWindowImmediately = false;
-	FSlateApplication::Get().AddWindow(RootWindowRef, bShowRootWindowImmediately);
-	FGlobalTabmanager::Get()->SetRootWindow(RootWindowRef);
-	FGlobalTabmanager::Get()->SetAllowWindowMenuBar(true);
-	FSlateNotificationManager::Get().SetRootWindow(RootWindowRef);
-	
-	const TSharedRef<FTabManager::FStack> MainStack = FTabManager::NewStack();
-	InitComponents(MainStack);
-	const TSharedRef<FTabManager::FArea> MainWindowArea = FTabManager::NewPrimaryArea();
-	const TSharedRef<FTabManager::FLayout> DefaultLayout = FTabManager::NewLayout("UnrealMultiUserServerLayout_v1.0");
-	MainWindowArea->Split(MainStack);
-	DefaultLayout->AddArea(MainWindowArea);
-	
-	PersistentLayout = FLayoutSaveRestore::LoadFromConfig(MultiUserServerLayoutIni, DefaultLayout);
-	TSharedPtr<SWidget> Content = FGlobalTabmanager::Get()->RestoreFrom(PersistentLayout.ToSharedRef(), RootWindow, bEmbedTitleAreaContent, EOutputCanBeNullptr::Never);
-	RootWindow->SetContent(Content.ToSharedRef());
-
-	RootWindow->SetOnWindowClosed(FOnWindowClosed::CreateRaw(this, &FConcertServerWindowController::OnWindowClosed));
-	RootWindow->ShowWindow();
-	constexpr bool bForceWindowToFront = true;
-	RootWindow->BringToFront(bForceWindowToFront);
-
-	RegisterForSessionDestructionEvents();
-	return RootWindowRef;
-}
+		const TSharedRef<FTabManager::FStack> MainStack = FTabManager::NewStack();
+		InitComponents(MainStack);
+		const TSharedRef<FTabManager::FArea> MainWindowArea = FTabManager::NewPrimaryArea();
+		const TSharedRef<FTabManager::FLayout> DefaultLayout = FTabManager::NewLayout("UnrealMultiUserServerLayout_v1.0");
+		MainWindowArea->Split(MainStack);
+		DefaultLayout->AddArea(MainWindowArea);
+		
+		PersistentLayout = FLayoutSaveRestore::LoadFromConfig(MultiUserServerLayoutIni, DefaultLayout);
+		TSharedPtr<SWidget> Content = FGlobalTabmanager::Get()->RestoreFrom(PersistentLayout.ToSharedRef(), RootWindow, bEmbedTitleAreaContent, EOutputCanBeNullptr::Never);
+		RootWindow->SetContent(Content.ToSharedRef());
 
-void FConcertServerWindowController::OpenSessionTab(const FGuid& SessionId)
-{
-	if (const TSharedPtr<FConcertSessionTabBase> SessionTab = GetOrRegisterSessionTab(SessionId))
-	{
-		SessionTab->OpenSessionTab();
-	}
-}
+		RootWindow->SetOnWindowClosed(FOnWindowClosed::CreateRaw(this, &FConcertServerWindowController::OnWindowClosed));
+		RootWindow->ShowWindow();
+		constexpr bool bForceWindowToFront = true;
+		RootWindow->BringToFront(bForceWindowToFront);
 
-void FConcertServerWindowController::DestroySessionTab(const FGuid& SessionId)
-{
-	// Destructor will handle the rest, e.g. removing the tab from the window
-	RegisteredSessions.Remove(SessionId);
-}
+		RegisterForSessionDestructionEvents();
+		return RootWindowRef;
+	}
 
-TSharedPtr<FConcertSessionTabBase> FConcertServerWindowController::GetOrRegisterSessionTab(const FGuid& SessionId)
-{
-	if (const TSharedRef<FConcertSessionTabBase>* FoundId = RegisteredSessions.Find(SessionId))
+	void FConcertServerWindowController::OpenSessionTab(const FGuid& SessionId)
 	{
-		return *FoundId;
+		if (const TSharedPtr<FConcertSessionTabBase> SessionTab = GetOrRegisterSessionTab(SessionId))
+		{
+			SessionTab->OpenSessionTab();
+		}
 	}
-	
-	if (const TSharedPtr<IConcertServerSession> Session = ServerInstance->GetConcertServer()->GetLiveSession(SessionId))
+
+	void FConcertServerWindowController::DestroySessionTab(const FGuid& SessionId)
 	{
-		const TSharedRef<FLiveConcertSessionTab> SessionTab = MakeShared<FLiveConcertSessionTab>(
-			Session.ToSharedRef(),
-			ServerInstance,
-			RootWindow.ToSharedRef(),
-			FLiveConcertSessionTab::FShowConnectedClients::CreateSP(this, &FConcertServerWindowController::ShowConnectedClients)
-			);
-		RegisteredSessions.Add(SessionId, SessionTab);
-		return SessionTab;
+		// Destructor will handle the rest, e.g. removing the tab from the window
+		RegisteredSessions.Remove(SessionId);
 	}
 
-	const bool bIsArchivedSession = ServerInstance->GetConcertServer()->GetArchivedSessionInfo(SessionId).IsSet();
-	if (bIsArchivedSession)
+	TSharedPtr<FConcertSessionTabBase> FConcertServerWindowController::GetOrRegisterSessionTab(const FGuid& SessionId)
 	{
-		const TSharedRef<FArchivedConcertSessionTab> SessionTab = MakeShared<FArchivedConcertSessionTab>(SessionId, ServerInstance, RootWindow.ToSharedRef());
-		RegisteredSessions.Add(SessionId, SessionTab);
-		return SessionTab;
+		if (const TSharedRef<FConcertSessionTabBase>* FoundId = RegisteredSessions.Find(SessionId))
+		{
+			return *FoundId;
+		}
+		
+		if (const TSharedPtr<IConcertServerSession> Session = ServerInstance->GetConcertServer()->GetLiveSession(SessionId))
+		{
+			const TSharedRef<FLiveConcertSessionTab> SessionTab = MakeShared<FLiveConcertSessionTab>(
+				Session.ToSharedRef(),
+				ServerInstance,
+				RootWindow.ToSharedRef(),
+				FLiveConcertSessionTab::FShowConnectedClients::CreateSP(this, &FConcertServerWindowController::ShowConnectedClients)
+				);
+			RegisteredSessions.Add(SessionId, SessionTab);
+			return SessionTab;
+		}
+
+		const bool bIsArchivedSession = ServerInstance->GetConcertServer()->GetArchivedSessionInfo(SessionId).IsSet();
+		if (bIsArchivedSession)
+		{
+			const TSharedRef<FArchivedConcertSessionTab> SessionTab = MakeShared<FArchivedConcertSessionTab>(SessionId, ServerInstance, RootWindow.ToSharedRef());
+			RegisteredSessions.Add(SessionId, SessionTab);
+			return SessionTab;
+		}
+		
+		return nullptr;
 	}
-	
-	return nullptr;
-}
 
-void FConcertServerWindowController::InitComponents(const TSharedRef<FTabManager::FStack>& MainArea)
-{
-	const FConcertComponentInitParams Params { ServerInstance, SharedThis(this), MainArea };
-	for (const TSharedRef<IConcertComponent>& ConcertComponent : ConcertComponents)
+	void FConcertServerWindowController::InitComponents(const TSharedRef<FTabManager::FStack>& MainArea)
 	{
-		ConcertComponent->Init(Params);
+		const FConcertComponentInitParams Params { ServerInstance, SharedThis(this), MainArea };
+		for (const TSharedRef<IConcertComponent>& ConcertComponent : ConcertComponents)
+		{
+			ConcertComponent->Init(Params);
+		}
 	}
-}
 
-void FConcertServerWindowController::RegisterForSessionDestructionEvents()
-{
-	ConcertServerEvents::OnLiveSessionDestroyed().AddSP(this, &FConcertServerWindowController::OnLiveSessionDestroyed);
-	ConcertServerEvents::OnArchivedSessionDestroyed().AddSP(this, &FConcertServerWindowController::OnArchivedSessionDestroyed);
-}
+	void FConcertServerWindowController::RegisterForSessionDestructionEvents()
+	{
+		ConcertServerEvents::OnLiveSessionDestroyed().AddSP(this, &FConcertServerWindowController::OnLiveSessionDestroyed);
+		ConcertServerEvents::OnArchivedSessionDestroyed().AddSP(this, &FConcertServerWindowController::OnArchivedSessionDestroyed);
+	}
 
-void FConcertServerWindowController::UnregisterFromSessionDestructionEvents() const
-{
-	ConcertServerEvents::OnLiveSessionDestroyed().RemoveAll(this);
-	ConcertServerEvents::OnArchivedSessionDestroyed().RemoveAll(this);
-}
+	void FConcertServerWindowController::UnregisterFromSessionDestructionEvents() const
+	{
+		ConcertServerEvents::OnLiveSessionDestroyed().RemoveAll(this);
+		ConcertServerEvents::OnArchivedSessionDestroyed().RemoveAll(this);
+	}
 
-void FConcertServerWindowController::OnLiveSessionDestroyed(const IConcertServer&, TSharedRef<IConcertServerSession> InLiveSession)
-{
-	DestroySessionTab(InLiveSession->GetId());
-}
+	void FConcertServerWindowController::OnLiveSessionDestroyed(const IConcertServer&, TSharedRef<IConcertServerSession> InLiveSession)
+	{
+		DestroySessionTab(InLiveSession->GetId());
+	}
 
-void FConcertServerWindowController::OnArchivedSessionDestroyed(const IConcertServer&, const FGuid& InArchivedSessionId)
-{
-	DestroySessionTab(InArchivedSessionId);
-}
+	void FConcertServerWindowController::OnArchivedSessionDestroyed(const IConcertServer&, const FGuid& InArchivedSessionId)
+	{
+		DestroySessionTab(InArchivedSessionId);
+	}
 
-void FConcertServerWindowController::ShowConnectedClients(const TSharedRef<IConcertServerSession>& ServerSession)
-{
-	ClientsController->ShowConnectedClients(ServerSession->GetId());
-}
+	void FConcertServerWindowController::ShowConnectedClients(const TSharedRef<IConcertServerSession>& ServerSession)
+	{
+		ClientsController->ShowConnectedClients(ServerSession->GetId());
+	}
 
-void FConcertServerWindowController::OnWindowClosed(const TSharedRef<SWindow>& Window)
-{
-	SaveLayout();
-	RootWindow.Reset();
-}
+	void FConcertServerWindowController::OnWindowClosed(const TSharedRef<SWindow>& Window)
+	{
+		SaveLayout();
+		RootWindow.Reset();
+	}
 
-void FConcertServerWindowController::SaveLayout() const
-{
-	FLayoutSaveRestore::SaveToConfig(MultiUserServerLayoutIni, PersistentLayout.ToSharedRef());
-    GConfig->Flush(false, MultiUserServerLayoutIni);
+	void FConcertServerWindowController::SaveLayout() const
+	{
+		FLayoutSaveRestore::SaveToConfig(MultiUserServerLayoutIni, PersistentLayout.ToSharedRef());
+	    GConfig->Flush(false, MultiUserServerLayoutIni);
+	}
 }
 
 #undef LOCTEXT_NAMESPACE 
\ No newline at end of file
diff --git a/Engine/Plugins/Developer/Concert/ConcertApp/MultiUserServer/Source/MultiUserServer/Private/Window/ConcertServerWindowController.h b/Engine/Plugins/Developer/Concert/ConcertApp/MultiUserServer/Source/MultiUserServer/Private/Window/ConcertServerWindowController.h
--- a/Engine/Plugins/Developer/Concert/ConcertApp/MultiUserServer/Source/MultiUserServer/Private/Window/ConcertServerWindowController.h
+++ b/Engine/Plugins/Developer/Concert/ConcertApp/MultiUserServer/Source/MultiUserServer/Private/Window/ConcertServerWindowController.h
@@ -5,83 +5,87 @@
 #include "CoreMinimal.h"
 #include "Framework/Docking/TabManager.h"
 
-class IConcertComponent;
 class IConcertServer;
 class FConcertSessionTabBase;
 class FLiveConcertSessionTab;
 class IConcertServerSession;
 class IConcertSyncServer;
-class FConcertServerSessionBrowserController;
-class FConcertClientsTabController;
 class FOutputLogController;
 class SWindow;
 
-struct FConcertServerWindowInitParams
+namespace UE::MultiUserServer
 {
-	/** The server that the window is supposed to manage */
-	TSharedRef<IConcertSyncServer> Server;
+	class IConcertComponent;
+	class FConcertServerSessionBrowserController;
+	class FConcertClientsTabController;
 	
-	/** Config path for server layout ini */
-	FString MultiUserServerLayoutIni;
-	
-	/** Additional modular concert features to enable */
-	TArray<TSharedRef<IConcertComponent>> AdditionalConcertComponents;
-	
-	FConcertServerWindowInitParams(TSharedRef<IConcertSyncServer> Server, FString MultiUserServerLayoutIni = FString())
-		: Server(Server)
-		, MultiUserServerLayoutIni(MoveTemp(MultiUserServerLayoutIni))
-	{}
-};
+	struct FConcertServerWindowInitParams
+	{
+		/** The server that the window is supposed to manage */
+		TSharedRef<IConcertSyncServer> Server;
+		
+		/** Config path for server layout ini */
+		FString MultiUserServerLayoutIni;
+		
+		/** Additional modular concert features to enable */
+		TArray<TSharedRef<IConcertComponent>> AdditionalConcertComponents;
+		
+		FConcertServerWindowInitParams(TSharedRef<IConcertSyncServer> Server, FString MultiUserServerLayoutIni = FString())
+			: Server(Server)
+			, MultiUserServerLayoutIni(MoveTemp(MultiUserServerLayoutIni))
+		{}
+	};
 
-/** Responsible for creating the Slate window for the server. Implements controller in the model-view-controller pattern. */
-class FConcertServerWindowController : public TSharedFromThis<FConcertServerWindowController>
-{
-public:
-	
-	FConcertServerWindowController(const FConcertServerWindowInitParams& Params);
-	~FConcertServerWindowController();
-	
-	TSharedRef<SWindow> CreateWindow();
-	TSharedPtr<SWindow> GetRootWindow() const { return RootWindow; }
+	/** Responsible for creating the Slate window for the server. Implements controller in the model-view-controller pattern. */
+	class FConcertServerWindowController : public TSharedFromThis<FConcertServerWindowController>
+	{
+	public:
+		
+		FConcertServerWindowController(const FConcertServerWindowInitParams& Params);
+		~FConcertServerWindowController();
+		
+		TSharedRef<SWindow> CreateWindow();
+		TSharedPtr<SWindow> GetRootWindow() const { return RootWindow; }
 
-	/** Opens or draws attention to the tab for the given live or archived session ID */
-	void OpenSessionTab(const FGuid& SessionId);
+		/** Opens or draws attention to the tab for the given live or archived session ID */
+		void OpenSessionTab(const FGuid& SessionId);
 
-	/** Destroys the tab associated with this live or archived session ID. */
-	void DestroySessionTab(const FGuid& SessionId);
-	
-private:
+		/** Destroys the tab associated with this live or archived session ID. */
+		void DestroySessionTab(const FGuid& SessionId);
+		
+	private:
 
-	/** The ini file to use for saving the layout */
-	FString MultiUserServerLayoutIni;
-	/** Holds the current layout for saving later. */
-	TSharedPtr<FTabManager::FLayout> PersistentLayout;
+		/** The ini file to use for saving the layout */
+		FString MultiUserServerLayoutIni;
+		/** Holds the current layout for saving later. */
+		TSharedPtr<FTabManager::FLayout> PersistentLayout;
 
-	TSharedRef<IConcertSyncServer> ServerInstance;
+		TSharedRef<IConcertSyncServer> ServerInstance;
 
-	/** The main window being managed */
-	TSharedPtr<SWindow> RootWindow;
-	TMap<FGuid, TSharedRef<FConcertSessionTabBase>> RegisteredSessions;
-	
-	/** Manages the session browser */
-	TSharedRef<FConcertServerSessionBrowserController> SessionBrowserController;
-	/** Manages the view of connected clients */
-	TSharedRef<FConcertClientsTabController> ClientsController;
-	/** Additional external concert components */
-	TArray<TSharedRef<IConcertComponent>> ConcertComponents;
-	
-	void InitComponents(const TSharedRef<FTabManager::FStack>& MainArea);
+		/** The main window being managed */
+		TSharedPtr<SWindow> RootWindow;
+		TMap<FGuid, TSharedRef<FConcertSessionTabBase>> RegisteredSessions;
+		
+		/** Manages the session browser */
+		TSharedRef<FConcertServerSessionBrowserController> SessionBrowserController;
+		/** Manages the view of connected clients */
+		TSharedRef<FConcertClientsTabController> ClientsController;
+		/** Additional external concert components */
+		TArray<TSharedRef<IConcertComponent>> ConcertComponents;
+		
+		void InitComponents(const TSharedRef<FTabManager::FStack>& MainArea);
 
-	/** Gets the manager for a session tab if the session ID is valid */
-	TSharedPtr<FConcertSessionTabBase> GetOrRegisterSessionTab(const FGuid& SessionId);
+		/** Gets the manager for a session tab if the session ID is valid */
+		TSharedPtr<FConcertSessionTabBase> GetOrRegisterSessionTab(const FGuid& SessionId);
 
-	void RegisterForSessionDestructionEvents();
-	void UnregisterFromSessionDestructionEvents() const;
-	void OnLiveSessionDestroyed(const IConcertServer&, TSharedRef<IConcertServerSession> InLiveSession);
-	void OnArchivedSessionDestroyed(const IConcertServer&, const FGuid& InArchivedSessionId);
+		void RegisterForSessionDestructionEvents();
+		void UnregisterFromSessionDestructionEvents() const;
+		void OnLiveSessionDestroyed(const IConcertServer&, TSharedRef<IConcertServerSession> InLiveSession);
+		void OnArchivedSessionDestroyed(const IConcertServer&, const FGuid& InArchivedSessionId);
 
-	void ShowConnectedClients(const TSharedRef<IConcertServerSession>& ServerSession);
-	
-	void OnWindowClosed(const TSharedRef<SWindow>& Window);
-	void SaveLayout() const;
-};
+		void ShowConnectedClients(const TSharedRef<IConcertServerSession>& ServerSession);
+		
+		void OnWindowClosed(const TSharedRef<SWindow>& Window);
+		void SaveLayout() const;
+	};
+}
diff --git a/Engine/Plugins/Developer/Concert/ConcertApp/MultiUserServer/Source/MultiUserServer/Private/Window/IConcertComponent.h b/Engine/Plugins/Developer/Concert/ConcertApp/MultiUserServer/Source/MultiUserServer/Private/Window/IConcertComponent.h
--- a/Engine/Plugins/Developer/Concert/ConcertApp/MultiUserServer/Source/MultiUserServer/Private/Window/IConcertComponent.h
+++ b/Engine/Plugins/Developer/Concert/ConcertApp/MultiUserServer/Source/MultiUserServer/Private/Window/IConcertComponent.h
@@ -5,34 +5,38 @@
 #include "CoreMinimal.h"
 #include "Framework/Docking/TabManager.h"
 
-class FConcertServerWindowController;
 class IConcertSyncServer;
 
-struct FConcertComponentInitParams
+namespace UE::MultiUserServer
 {
-	/** The server being managed */
-	TSharedRef<IConcertSyncServer> Server;
+	class FConcertServerWindowController;
 
-	/** Manages the server slate application */
-	TSharedRef<FConcertServerWindowController> WindowController;
+	struct FConcertComponentInitParams
+	{
+		/** The server being managed */
+		TSharedRef<IConcertSyncServer> Server;
 
-	/** The root area of the main window layout. Add tabs to this. */
-	TSharedRef<FTabManager::FStack> MainStack;
+		/** Manages the server slate application */
+		TSharedRef<FConcertServerWindowController> WindowController;
 
-	FConcertComponentInitParams(TSharedRef<IConcertSyncServer> Server, TSharedRef<FConcertServerWindowController> WindowController, TSharedRef<FTabManager::FStack> MainStack)
-		: Server(MoveTemp(Server))
-		, WindowController(MoveTemp(WindowController))
-		, MainStack(MoveTemp(MainStack))
-	{}
-};
+		/** The root area of the main window layout. Add tabs to this. */
+		TSharedRef<FTabManager::FStack> MainStack;
 
-/** Provides the base interface for elements in the concert server UI. */
-class IConcertComponent
-{
-public:
+		FConcertComponentInitParams(TSharedRef<IConcertSyncServer> Server, TSharedRef<FConcertServerWindowController> WindowController, TSharedRef<FTabManager::FStack> MainStack)
+			: Server(MoveTemp(Server))
+			, WindowController(MoveTemp(WindowController))
+			, MainStack(MoveTemp(MainStack))
+		{}
+	};
+
+	/** Provides the base interface for elements in the concert server UI. */
+	class IConcertComponent
+	{
+	public:
 
-	virtual ~IConcertComponent() = default;
+		virtual ~IConcertComponent() = default;
 
-	/** Initialises the component, e.g. registering tab spawners, etc. */
-	virtual void Init(const FConcertComponentInitParams& Params) {}
-};
+		/** Initialises the component, e.g. registering tab spawners, etc. */
+		virtual void Init(const FConcertComponentInitParams& Params) {}
+	};
+}
\ No newline at end of file
