Fix build issue

diff --git a/Engine/Plugins/Developer/Concert/ConcertApp/MultiUserServer/Source/MultiUserServer/Private/Widgets/Util/Filter/ConcertFilter.h b/Engine/Plugins/Developer/Concert/ConcertApp/MultiUserServer/Source/MultiUserServer/Private/Widgets/Util/Filter/ConcertFilter.h
--- a/Engine/Plugins/Developer/Concert/ConcertApp/MultiUserServer/Source/MultiUserServer/Private/Widgets/Util/Filter/ConcertFilter.h
+++ b/Engine/Plugins/Developer/Concert/ConcertApp/MultiUserServer/Source/MultiUserServer/Private/Widgets/Util/Filter/ConcertFilter.h
@@ -19,7 +19,7 @@ namespace UE::MultiUserServer
 	public:
 
 		//~ Begin IFilter Interface
-		DECLARE_DERIVED_EVENT( FFrontendFilter, IFilter<TFilterType>::template FChangedEvent, FChangedEvent );
+		DECLARE_DERIVED_EVENT( TConcertFilter<TFilterType>, IFilter<TFilterType>::FChangedEvent, FChangedEvent );
 		virtual FChangedEvent& OnChanged() override { return ChangedEvent; }
 		//~ End IFilter Interface
 
