Fix build issue

diff --git a/Engine/Plugins/Developer/Concert/ConcertApp/MultiUserServer/Source/MultiUserServer/Private/Widgets/Util/Filter/ConcertFrontendFilter.h b/Engine/Plugins/Developer/Concert/ConcertApp/MultiUserServer/Source/MultiUserServer/Private/Widgets/Util/Filter/ConcertFrontendFilter.h
--- a/Engine/Plugins/Developer/Concert/ConcertApp/MultiUserServer/Source/MultiUserServer/Private/Widgets/Util/Filter/ConcertFrontendFilter.h
+++ b/Engine/Plugins/Developer/Concert/ConcertApp/MultiUserServer/Source/MultiUserServer/Private/Widgets/Util/Filter/ConcertFrontendFilter.h
@@ -51,7 +51,7 @@ namespace UE::MultiUserServer
 		{
 			Implementation.OnChanged().AddLambda([this]()
 			{
-				Super::template OnChanged().Broadcast();
+				Super::OnChanged().Broadcast();
 			});
 		}
 
diff --git a/Engine/Plugins/Developer/Concert/ConcertApp/MultiUserServer/Source/MultiUserServer/Private/Widgets/Util/Filter/ConcertFrontendFilter_TextSearch.h b/Engine/Plugins/Developer/Concert/ConcertApp/MultiUserServer/Source/MultiUserServer/Private/Widgets/Util/Filter/ConcertFrontendFilter_TextSearch.h
--- a/Engine/Plugins/Developer/Concert/ConcertApp/MultiUserServer/Source/MultiUserServer/Private/Widgets/Util/Filter/ConcertFrontendFilter_TextSearch.h
+++ b/Engine/Plugins/Developer/Concert/ConcertApp/MultiUserServer/Source/MultiUserServer/Private/Widgets/Util/Filter/ConcertFrontendFilter_TextSearch.h
@@ -19,11 +19,11 @@ namespace UE::MultiUserServer
 	public:
 
 		TConcertFilter_TextSearch()
-			: TextFilter(TTextFilter<TFilterType>::template FItemToStringArray::CreateRaw(this, &TConcertFilter_TextSearch<TFilterType>::GenerateSearchTerms))
+			: TextFilter(TTextFilter<TFilterType>::FItemToStringArray::CreateRaw(this, &TConcertFilter_TextSearch<TFilterType>::GenerateSearchTerms))
 		{
 			TextFilter.OnChanged().AddLambda([this]
 			{
-				Super::template OnChanged().Broadcast();
+				Super::OnChanged().Broadcast();
 			});
 		}
 	
@@ -55,12 +55,11 @@ namespace UE::MultiUserServer
 		TConcertFrontendFilter_TextSearch(TArg&&... Arg)
 			: Super(Forward<TArg>(Arg)...)
 		{
-			Super::template ChildSlot = SNew(SSearchBox)
+			Super::ChildSlot = SNew(SSearchBox)
 				.OnTextChanged_Lambda([this](const FText& NewSearchText)
 				{
 					SearchText = NewSearchText;
-					TTextSearchFilterType& Impl = Super::template Implementation;
-					Impl.SetRawFilterText(NewSearchText);
+					Super::Implementation.SetRawFilterText(NewSearchText);
 				})
 				.DelayChangeNotificationsWhileTyping(true);
 		}
diff --git a/Engine/Plugins/Developer/Concert/ConcertApp/MultiUserServer/Source/MultiUserServer/Private/Widgets/Util/Filter/ConcertRootFrontendFilter.h b/Engine/Plugins/Developer/Concert/ConcertApp/MultiUserServer/Source/MultiUserServer/Private/Widgets/Util/Filter/ConcertRootFrontendFilter.h
--- a/Engine/Plugins/Developer/Concert/ConcertApp/MultiUserServer/Source/MultiUserServer/Private/Widgets/Util/Filter/ConcertRootFrontendFilter.h
+++ b/Engine/Plugins/Developer/Concert/ConcertApp/MultiUserServer/Source/MultiUserServer/Private/Widgets/Util/Filter/ConcertRootFrontendFilter.h
@@ -93,7 +93,7 @@ namespace UE::MultiUserServer
 
 		void BroadcastOnChanged()
 		{
-			Super::template OnChanged().Broadcast();
+			Super::OnChanged().Broadcast();
 		}
 
 		/** Builds the widgets that go under the text */
