Fix static warning

diff --git a/Engine/Plugins/Developer/Concert/ConcertApp/MultiUserServer/Source/MultiUserServer/Private/Widgets/Clients/PackageTransmission/Model/PackageTransmissionModel.cpp b/Engine/Plugins/Developer/Concert/ConcertApp/MultiUserServer/Source/MultiUserServer/Private/Widgets/Clients/PackageTransmission/Model/PackageTransmissionModel.cpp
--- a/Engine/Plugins/Developer/Concert/ConcertApp/MultiUserServer/Source/MultiUserServer/Private/Widgets/Clients/PackageTransmission/Model/PackageTransmissionModel.cpp
+++ b/Engine/Plugins/Developer/Concert/ConcertApp/MultiUserServer/Source/MultiUserServer/Private/Widgets/Clients/PackageTransmission/Model/PackageTransmissionModel.cpp
@@ -81,9 +81,8 @@ namespace UE::MultiUserServer
 		const bool bHadBeginSendData = InTransmissionToIndex.RemoveAndCopyValue(Args.TransmissionId, Key);
 		if (ensure(bHadBeginSendData) && ensure(Entries.IsValidIndex(Key)))
 		{
-			TSharedPtr<FPackageTransmissionEntry>& Entry = Entries[Key];
-			Entry->MessageId = Args.MessageId;
-			Entry->TransmissionState = TransmissionState;
+			Entries[Key]->MessageId = Args.MessageId;
+			Entries[Key]->TransmissionState = TransmissionState;
 
 			ModifiedSinceLastTick.Add(Key);
 		}
