Fix a static warning

diff --git a/Engine/Plugins/Developer/Concert/ConcertApp/MultiUserServer/Source/MultiUserServer/Private/Widgets/Clients/PackageTransmission/Model/FilteredPackageTransmissionModel.cpp b/Engine/Plugins/Developer/Concert/ConcertApp/MultiUserServer/Source/MultiUserServer/Private/Widgets/Clients/PackageTransmission/Model/FilteredPackageTransmissionModel.cpp
--- a/Engine/Plugins/Developer/Concert/ConcertApp/MultiUserServer/Source/MultiUserServer/Private/Widgets/Clients/PackageTransmission/Model/FilteredPackageTransmissionModel.cpp
+++ b/Engine/Plugins/Developer/Concert/ConcertApp/MultiUserServer/Source/MultiUserServer/Private/Widgets/Clients/PackageTransmission/Model/FilteredPackageTransmissionModel.cpp
@@ -76,14 +76,14 @@ namespace UE::MultiUserServer
 	void FFilteredPackageTransmissionModel::ApplyFilter(bool bSilent, int32 StartIndex)
 	{
 		uint32 NumAdded = 0;
-		
-		for (int32 i = StartIndex; i < RealSource->GetEntries().Num(); ++i)
+
+		const TArray<TSharedPtr<FPackageTransmissionEntry>>& Entries = RealSource->GetEntries();
+		for (int32 i = StartIndex; i < Entries.Num(); ++i)
 		{
-			const TSharedPtr<FPackageTransmissionEntry>& Entry = RealSource->GetEntries()[i];
-			if (Filter->PassesFilter(*Entry))
+			if (Filter->PassesFilter(*Entries[i]))
 			{
-				FPackageTransmissionId Index = FilteredEntries.Add(Entry);
-				RebindRealSourceToFilteredIndex.Add(Entry->TransmissionId, Index);
+				FPackageTransmissionId Index = FilteredEntries.Add(Entries[i]);
+				RebindRealSourceToFilteredIndex.Add(Entries[i]->TransmissionId, Index);
 				++NumAdded;
 			}
 		}
