Refactor frontend filters to use FFilterBase

diff --git a/Engine/Plugins/Developer/Concert/ConcertApp/MultiUserServer/Source/MultiUserServer/MultiUserServer.build.cs b/Engine/Plugins/Developer/Concert/ConcertApp/MultiUserServer/Source/MultiUserServer/MultiUserServer.build.cs
--- a/Engine/Plugins/Developer/Concert/ConcertApp/MultiUserServer/Source/MultiUserServer/MultiUserServer.build.cs
+++ b/Engine/Plugins/Developer/Concert/ConcertApp/MultiUserServer/Source/MultiUserServer/MultiUserServer.build.cs
@@ -26,7 +26,8 @@ namespace UnrealBuildTool.Rules
 					"SlateCore",
 					"StandaloneRenderer",	
 					"ToolWidgets",
-					"WorkspaceMenuStructure" // Needed so OutputLog module to work correctly
+					"ToolMenus",
+					"WorkspaceMenuStructure" // Needed for OutputLog module to work correctly
 				}
 			);
 			
diff --git a/Engine/Plugins/Developer/Concert/ConcertApp/MultiUserServer/Source/MultiUserServer/Private/Widgets/Clients/Logging/Filter/ConcertFrontendLogFilter_Ack.h b/Engine/Plugins/Developer/Concert/ConcertApp/MultiUserServer/Source/MultiUserServer/Private/Widgets/Clients/Logging/Filter/ConcertFrontendLogFilter_Ack.h
--- a/Engine/Plugins/Developer/Concert/ConcertApp/MultiUserServer/Source/MultiUserServer/Private/Widgets/Clients/Logging/Filter/ConcertFrontendLogFilter_Ack.h
+++ b/Engine/Plugins/Developer/Concert/ConcertApp/MultiUserServer/Source/MultiUserServer/Private/Widgets/Clients/Logging/Filter/ConcertFrontendLogFilter_Ack.h
@@ -49,9 +49,12 @@ namespace UE::MultiUserServer
 	{
 	public:
 
-		FConcertFrontendLogFilter_Ack()
-			: TConcertFrontendLogFilter_BaseSetSelection<FConcertLogFilter_Ack>(NSLOCTEXT("UnrealMultiUserUI.Filter.Ack", "Name", "ACKs"))
+		FConcertFrontendLogFilter_Ack(TSharedRef<FFilterCategory> FilterCategory)
+			: TConcertFrontendLogFilter_BaseSetSelection<FConcertLogFilter_Ack>(MoveTemp(FilterCategory), NSLOCTEXT("UnrealMultiUserUI.Filter.Ack", "Name", "ACKs"))
 		{}
+		
+		virtual FString GetName() const override { return TEXT("Acknowledgement"); }
+		virtual FText GetDisplayName() const override { return NSLOCTEXT("UnrealMultiUserUI.FConcertFrontendLogFilter_Ack", "DisplayLabel", "Acknowledgement"); }
 	};
 }
 
diff --git a/Engine/Plugins/Developer/Concert/ConcertApp/MultiUserServer/Source/MultiUserServer/Private/Widgets/Clients/Logging/Filter/ConcertFrontendLogFilter_BaseSetSelection.h b/Engine/Plugins/Developer/Concert/ConcertApp/MultiUserServer/Source/MultiUserServer/Private/Widgets/Clients/Logging/Filter/ConcertFrontendLogFilter_BaseSetSelection.h
--- a/Engine/Plugins/Developer/Concert/ConcertApp/MultiUserServer/Source/MultiUserServer/Private/Widgets/Clients/Logging/Filter/ConcertFrontendLogFilter_BaseSetSelection.h
+++ b/Engine/Plugins/Developer/Concert/ConcertApp/MultiUserServer/Source/MultiUserServer/Private/Widgets/Clients/Logging/Filter/ConcertFrontendLogFilter_BaseSetSelection.h
@@ -3,14 +3,9 @@
 #pragma once
 
 #include "CoreMinimal.h"
-
-#include "ConcertLogFilterTypes.h"
-#include "Widgets/Util/Filter/ConcertFilter.h"
+#include "Misc/IFilter.h"
 #include "Widgets/Util/Filter/ConcertFrontendFilter.h"
-
 #include "Framework/MultiBox/MultiBoxBuilder.h"
-#include "Widgets/Input/SComboButton.h"
-#include "Widgets/SBoxPanel.h"
 
 #define LOCTEXT_NAMESPACE "UnrealMultiUserUI.BaseSetSelection"
 
@@ -23,8 +18,9 @@ namespace UE::MultiUserServer
 	 * - static FString GetOptionDisplayString(const TSetType&);
 	 */
 	template<typename TRealFilterImpl, typename TSetType>
-	class TConcertLogFilter_BaseSetSelection : public FConcertLogFilter
+	class TConcertLogFilter_BaseSetSelection : public IFilter<const FConcertLogEntry&>
 	{
+		using Super = IFilter<TRealFilterImpl>;
 	public:
 
 		using TItemType = TSetType;
@@ -96,57 +92,28 @@ namespace UE::MultiUserServer
 			return AllowedItems.Num() == TRealFilterImpl::GetAllOptions().Num();
 		}
 		uint8 GetNumSelected() const { return AllowedItems.Num(); }
-		
+
+		virtual FChangedEvent& OnChanged() override { return ChangedEvent; }
+
 	private:
 	
 		TSet<TSetType> AllowedItems;
+		FChangedEvent ChangedEvent;
 	};
 
 	template<typename TRealFilterImpl>
 	class TConcertFrontendLogFilter_BaseSetSelection : public TConcertFrontendFilterAggregate<TRealFilterImpl, const FConcertLogEntry&>
 	{
 		using TItemType = typename TRealFilterImpl::TItemType;
+		using Super = TConcertFrontendFilterAggregate<TRealFilterImpl, const FConcertLogEntry&>;
 	public:
 
-		TConcertFrontendLogFilter_BaseSetSelection(FText FilterName)
-		{
-			this->ChildSlot = SNew(SHorizontalBox)
-				.ToolTipText(LOCTEXT("BaseSetSelection.ToolTipText", "Select a list of allowed items\nHint: Type in menu to search"))
-
-				+SHorizontalBox::Slot()
-				.AutoWidth()
-				.VAlign(VAlign_Center)
-				[
-					SNew(STextBlock)
-					.Text(FilterName)
-				]
-
-				+SHorizontalBox::Slot()
-				.AutoWidth()
-				.Padding(2, 0, 0, 0)
-				.VAlign(VAlign_Center)
-				[
-					SNew(SComboButton)
-					.OnGetMenuContent_Raw(this, &TConcertFrontendLogFilter_BaseSetSelection::MakeSelectionMenu)
-					.ButtonContent()
-					[
-						SNew(STextBlock)
-						.Text_Lambda([this]()
-						{
-							return this->Implementation.AreAllAllowed()
-								? LOCTEXT("BaseSetSelection.Selection.All", "All")
-								: FText::FromString(FString::FromInt(this->Implementation.GetNumSelected()));
-						})
-					]
-				];
-		}
-
-	private:
+		TConcertFrontendLogFilter_BaseSetSelection(TSharedRef<FFilterCategory> FilterCategory, FText FilterName)
+			: Super(MoveTemp(FilterCategory))
+		{}
 
-		TSharedRef<SWidget> MakeSelectionMenu()
+		virtual void ExposeEditWidgets(FMenuBuilder& MenuBuilder) override
 		{
-			FMenuBuilder MenuBuilder(false, nullptr);
-
 			MenuBuilder.AddMenuEntry(
 				LOCTEXT("BaseSetSelection.SelectAll.", "Select all"),
 				LOCTEXT("BaseSetSelection.SelectAll.Tooltip", "Allows all items"),
@@ -170,8 +137,6 @@ namespace UE::MultiUserServer
 				EUserInterfaceActionType::Button
 				);
 	
-			MenuBuilder.AddSeparator();
-	
 			for (const TItemType& Item : TRealFilterImpl::GetAllOptions())
 			{
 				MenuBuilder.AddMenuEntry(
@@ -196,7 +161,6 @@ namespace UE::MultiUserServer
 					EUserInterfaceActionType::ToggleButton
 				);
 			}
-			return MenuBuilder.MakeWidget();
 		}
 	};
 
diff --git a/Engine/Plugins/Developer/Concert/ConcertApp/MultiUserServer/Source/MultiUserServer/Private/Widgets/Clients/Logging/Filter/ConcertFrontendLogFilter_Client.h b/Engine/Plugins/Developer/Concert/ConcertApp/MultiUserServer/Source/MultiUserServer/Private/Widgets/Clients/Logging/Filter/ConcertFrontendLogFilter_Client.h
--- a/Engine/Plugins/Developer/Concert/ConcertApp/MultiUserServer/Source/MultiUserServer/Private/Widgets/Clients/Logging/Filter/ConcertFrontendLogFilter_Client.h
+++ b/Engine/Plugins/Developer/Concert/ConcertApp/MultiUserServer/Source/MultiUserServer/Private/Widgets/Clients/Logging/Filter/ConcertFrontendLogFilter_Client.h
@@ -3,15 +3,15 @@
 #pragma once
 
 #include "CoreMinimal.h"
-#include "ConcertLogFilterTypes.h"
-#include "Widgets/Util/Filter/ConcertFilter.h"
+#include "Misc/IFilter.h"
+#include "Widgets/Clients/Logging/ConcertLogEntry.h"
 
 class FEndpointToUserNameCache;
 
 namespace UE::MultiUserServer
 {
 	/** Only allows messages from the given clients */
-	class FConcertLogFilter_Client : public FConcertLogFilter
+	class FConcertLogFilter_Client : public IFilter<const FConcertLogEntry&>
 	{
 	public:
 
@@ -22,6 +22,7 @@ namespace UE::MultiUserServer
 
 		//~ Begin FConcertLogFilter Interface
 		virtual bool PassesFilter(const FConcertLogEntry& InItem) const override;
+		virtual FChangedEvent& OnChanged() override { return ChangedEvent; }
 		//~ End FConcertLogFilter Interface
 
 	private:
@@ -31,5 +32,7 @@ namespace UE::MultiUserServer
 
 		/** Translates Concert endpoint Ids to messaging node Ids */
 		TSharedRef<FEndpointToUserNameCache> EndpointCache;
+
+		FChangedEvent ChangedEvent;
 	};
 }
diff --git a/Engine/Plugins/Developer/Concert/ConcertApp/MultiUserServer/Source/MultiUserServer/Private/Widgets/Clients/Logging/Filter/ConcertFrontendLogFilter_MessageAction.h b/Engine/Plugins/Developer/Concert/ConcertApp/MultiUserServer/Source/MultiUserServer/Private/Widgets/Clients/Logging/Filter/ConcertFrontendLogFilter_MessageAction.h
--- a/Engine/Plugins/Developer/Concert/ConcertApp/MultiUserServer/Source/MultiUserServer/Private/Widgets/Clients/Logging/Filter/ConcertFrontendLogFilter_MessageAction.h
+++ b/Engine/Plugins/Developer/Concert/ConcertApp/MultiUserServer/Source/MultiUserServer/Private/Widgets/Clients/Logging/Filter/ConcertFrontendLogFilter_MessageAction.h
@@ -6,7 +6,6 @@
 #include "ConcertFrontendLogFilter_BaseSetSelection.h"
 #include "Widgets/Clients/Logging/ConcertLogEntry.h"
 #include "Widgets/Clients/Logging/Util/MessageActionUtils.h"
-#include "Widgets/Util/Filter/ConcertFilter.h"
 #include "Widgets/Util/Filter/ConcertFrontendFilter.h"
 
 namespace UE::MultiUserServer
@@ -38,9 +37,12 @@ namespace UE::MultiUserServer
 	{
 	public:
 
-		FConcertFrontendLogFilter_MessageAction()
-			: TConcertFrontendLogFilter_BaseSetSelection<FConcertLogFilter_MessageAction>(NSLOCTEXT("UnrealMultiUserUI.Filter.MessageAction", "Name", "Actions"))
+		FConcertFrontendLogFilter_MessageAction(TSharedRef<FFilterCategory> FilterCategory)
+			: TConcertFrontendLogFilter_BaseSetSelection<FConcertLogFilter_MessageAction>(MoveTemp(FilterCategory), NSLOCTEXT("UnrealMultiUserUI.Filter.MessageAction", "Name", "Actions"))
 		{}
+		
+		virtual FString GetName() const override { return TEXT("Message Action"); }
+		virtual FText GetDisplayName() const override { return NSLOCTEXT("UnrealMultiUserUI.FConcertFrontendLogFilter_MessageAction", "DisplayLabel", "Action"); }
 	};
 }
 
diff --git a/Engine/Plugins/Developer/Concert/ConcertApp/MultiUserServer/Source/MultiUserServer/Private/Widgets/Clients/Logging/Filter/ConcertFrontendLogFilter_MessageType.h b/Engine/Plugins/Developer/Concert/ConcertApp/MultiUserServer/Source/MultiUserServer/Private/Widgets/Clients/Logging/Filter/ConcertFrontendLogFilter_MessageType.h
--- a/Engine/Plugins/Developer/Concert/ConcertApp/MultiUserServer/Source/MultiUserServer/Private/Widgets/Clients/Logging/Filter/ConcertFrontendLogFilter_MessageType.h
+++ b/Engine/Plugins/Developer/Concert/ConcertApp/MultiUserServer/Source/MultiUserServer/Private/Widgets/Clients/Logging/Filter/ConcertFrontendLogFilter_MessageType.h
@@ -6,7 +6,6 @@
 #include "ConcertFrontendLogFilter_BaseSetSelection.h"
 #include "Widgets/Clients/Logging/ConcertLogEntry.h"
 #include "Widgets/Clients/Logging/Util/MessageTypeUtils.h"
-#include "Widgets/Util/Filter/ConcertFilter.h"
 #include "Widgets/Util/Filter/ConcertFrontendFilter.h"
 
 namespace UE::MultiUserServer
@@ -42,9 +41,12 @@ namespace UE::MultiUserServer
 	{
 	public:
 
-		FConcertFrontendLogFilter_MessageType()
-			: TConcertFrontendLogFilter_BaseSetSelection<FConcertLogFilter_MessageType>(NSLOCTEXT("UnrealMultiUserUI.Filter.MessageType", "Name", "Message Type"))
+		FConcertFrontendLogFilter_MessageType(TSharedRef<FFilterCategory> FilterCategory)
+			: TConcertFrontendLogFilter_BaseSetSelection<FConcertLogFilter_MessageType>(MoveTemp(FilterCategory), NSLOCTEXT("UnrealMultiUserUI.Filter.MessageType", "Name", "Message Type"))
 		{}
+		
+		virtual FString GetName() const override { return TEXT("Message Type"); }
+		virtual FText GetDisplayName() const override { return NSLOCTEXT("UnrealMultiUserUI.FConcertFrontendLogFilter_MessageType", "DisplayLabel", "Type"); }
 	};
 }
 
diff --git a/Engine/Plugins/Developer/Concert/ConcertApp/MultiUserServer/Source/MultiUserServer/Private/Widgets/Clients/Logging/Filter/ConcertFrontendLogFilter_Size.cpp b/Engine/Plugins/Developer/Concert/ConcertApp/MultiUserServer/Source/MultiUserServer/Private/Widgets/Clients/Logging/Filter/ConcertFrontendLogFilter_Size.cpp
--- a/Engine/Plugins/Developer/Concert/ConcertApp/MultiUserServer/Source/MultiUserServer/Private/Widgets/Clients/Logging/Filter/ConcertFrontendLogFilter_Size.cpp
+++ b/Engine/Plugins/Developer/Concert/ConcertApp/MultiUserServer/Source/MultiUserServer/Private/Widgets/Clients/Logging/Filter/ConcertFrontendLogFilter_Size.cpp
@@ -4,13 +4,11 @@
 
 #include "ConcertTransportEvents.h"
 #include "SSimpleComboButton.h"
+#include "Widgets/Clients/Logging/ConcertLogEntry.h"
+#include "Widgets/Util/Filter/ConcertFilterUtils.h"
 
 #include "Framework/MultiBox/MultiBoxBuilder.h"
-#include "Widgets/SBoxPanel.h"
-#include "Widgets/Clients/Logging/ConcertLogEntry.h"
-#include "Widgets/Input/SButton.h"
 #include "Widgets/Input/SNumericEntryBox.h"
-#include "Widgets/Text/STextBlock.h"
 
 #define LOCTEXT_NAMESPACE "UnrealMultiUserUI.FConcertLogFilter_Size"
 
@@ -34,21 +32,13 @@ namespace UE::MultiUserServer
 		}
 	}
 
-	void FConcertLogFilter_Size::AdvanceFilterMode()
+	void FConcertLogFilter_Size::SetOperator(ESizeFilterMode Operator)
 	{
-		switch (FilterMode)
+		if (FilterMode != Operator)
 		{
-		case ESizeFilterMode::LessThanOrEqual:
-			FilterMode = ESizeFilterMode::BiggerThanOrEqual;
-			break;
-		case ESizeFilterMode::BiggerThanOrEqual:
-			FilterMode = ESizeFilterMode::LessThanOrEqual;
-			break;
-		default:
-			checkNoEntry();
+			FilterMode = Operator;
+			OnChanged().Broadcast();
 		}
-
-		OnChanged().Broadcast();
 	}
 
 	void FConcertLogFilter_Size::SetSizeInBytes(uint32 NewSizeInBytes)
@@ -69,97 +59,55 @@ namespace UE::MultiUserServer
 		}
 	}
 
-	FConcertFrontendLogFilter_Size::FConcertFrontendLogFilter_Size()
-	{
-		ChildSlot = SNew(SHorizontalBox)
-			+SHorizontalBox::Slot()
-			.AutoWidth()
-			.HAlign(HAlign_Center)
-			[
-				SNew(SButton)
-				.OnClicked_Lambda([this]()
-				{
-					Implementation.AdvanceFilterMode();
-					return FReply::Handled();
-				})
-				.ToolTipText_Lambda([this]()
-				{
-					switch (Implementation.GetFilterMode())
-					{
-					case ESizeFilterMode::BiggerThanOrEqual: return FText::Format(LOCTEXT("BiggerThanOrEqual.ToolTipFmt", "Size >= {0}"), GetSizeAndUnitAsText());
-						case ESizeFilterMode::LessThanOrEqual: return FText::Format(LOCTEXT("LessThanOrEqual.ToolTipFmt", "Size <= {0}"), GetSizeAndUnitAsText());
-						default: return FText::GetEmpty();
-					}
-				})
-				[
-					SNew(STextBlock)
-					.Text_Lambda([this]()
-					{
-						switch (Implementation.GetFilterMode())
-						{
-							case ESizeFilterMode::BiggerThanOrEqual: return LOCTEXT("BiggerThanOrEqual.Text", ">=") ;
-							case ESizeFilterMode::LessThanOrEqual: return LOCTEXT("LessThanOrEqual.Text", "<=") ;
-							default: return FText::GetEmpty();
-						}
-					})
-				]
-			]
-		
-			+SHorizontalBox::Slot()
-			.AutoWidth()
-			[
-				SNew(SNumericEntryBox<uint32>)
-				.AllowSpin(true)
-				.MinDesiredValueWidth(30)
-				.MaxSliderValue(1000)
-				.OnValueChanged_Lambda([this](uint32 NewValue){ Implementation.SetSizeInBytes(NewValue); })
-				.Value_Lambda([this](){ return Implementation.GetSizeInBytes(); })
-			]
-
-			+SHorizontalBox::Slot()
-			.AutoWidth()
-			[
-				SNew(SSimpleComboButton)
-				.HasDownArrow(true)
-				.Text_Lambda([this](){ return FText::FromString(FUnitConversion::GetUnitDisplayString(Implementation.GetDataUnit())); })
-				.MenuContent()
-				[
-					MakeDataUnitMenu()
-				]
-			];
-	}
+	FConcertFrontendLogFilter_Size::FConcertFrontendLogFilter_Size(TSharedRef<FFilterCategory> FilterCategory)
+		: Super(MoveTemp(FilterCategory))
+	{}
 
-	TSharedRef<SWidget> FConcertFrontendLogFilter_Size::MakeDataUnitMenu()
+	void FConcertFrontendLogFilter_Size::ExposeEditWidgets(FMenuBuilder& MenuBuilder)
 	{
-		FMenuBuilder MenuBuilder(true, nullptr);
+		using namespace ConcertFilterUtils;
 
-		for (const EUnit Unit : Implementation.GetAllowedUnits())
-		{
-			MenuBuilder.AddMenuEntry(
-				FText::FromString(FUnitConversion::GetUnitDisplayString(Unit)),
-				FText::GetEmpty(),
-				FSlateIcon(),
-				FUIAction(
-					FExecuteAction::CreateLambda([this, Unit]()
-					{
-						Implementation.SetDataUnit(Unit);
-					}),
-					FCanExecuteAction::CreateLambda([] { return true; })),
-				NAME_None,
-				EUserInterfaceActionType::Button
-				);
-		}
-		
-		return MenuBuilder.MakeWidget();
-	}
+		constexpr float MenuWidth = 100.f;
+		TRadialMenuBuilder<ESizeFilterMode>::AddRadialSubMenu(
+			MenuBuilder,
+			LOCTEXT("Operator", "Operator"),
+			TAttribute<ESizeFilterMode>::CreateLambda([this](){ return Implementation.GetOperator(); }),
+			TArray<ESizeFilterMode>{ ESizeFilterMode::BiggerThanOrEqual, ESizeFilterMode::LessThanOrEqual },
+			TRadialMenuBuilder<ESizeFilterMode>::FItemToText::CreateLambda([this](const ESizeFilterMode& Operator)
+			{
+				switch (Operator)
+				{
+				case ESizeFilterMode::BiggerThanOrEqual: return LOCTEXT("BiggerThanOrEqual.Text", ">= (Greater)") ;
+				case ESizeFilterMode::LessThanOrEqual: return LOCTEXT("LessThanOrEqual.Text", "<= (Less)") ;
+				default: return FText::GetEmpty();
+				}
+			}),
+			TRadialMenuBuilder<ESizeFilterMode>::FOnSelectItem::CreateLambda([this](const ESizeFilterMode& Operator)
+			{
+				Implementation.SetOperator(Operator);
+			}), MenuWidth);
 
-	FText FConcertFrontendLogFilter_Size::GetSizeAndUnitAsText() const
-	{
-		return FText::Format(
-			LOCTEXT("SizeAndUnitAsTextFmt", "{0} {1}"),
-			FUnitConversion::Convert(Implementation.GetSizeInBytes(), Implementation.GetDataUnit(), EUnit::Bytes),
-			FText::FromString(FUnitConversion::GetUnitDisplayString(Implementation.GetDataUnit()))
-			);
+		const TSharedRef<SWidget> SizeWidget = SNew(SNumericEntryBox<uint32>)
+			.AllowSpin(true)
+			.MinDesiredValueWidth(30)
+			.MaxSliderValue(1000)
+			.OnValueChanged_Lambda([this](uint32 NewValue){ Implementation.SetSizeInBytes(NewValue); })
+			.Value_Lambda([this](){ return Implementation.GetSizeInBytes(); });
+		MenuBuilder.AddWidget(SetMenuWidgetWidth(SizeWidget), LOCTEXT("Size", "Size"), true);
+
+		TRadialMenuBuilder<EUnit>::AddRadialSubMenu(
+			MenuBuilder,
+			LOCTEXT("Unit", "Unit"),
+			TAttribute<EUnit>::CreateLambda([this](){ return Implementation.GetDataUnit(); }),
+			TArray<EUnit>{ EUnit::Bytes, EUnit::Kilobytes, EUnit::Megabytes },
+			TRadialMenuBuilder<EUnit>::FItemToText::CreateLambda([this](const EUnit& Unit)
+			{
+				return FText::FromString(FUnitConversion::GetUnitDisplayString(Unit));
+			}),
+			TRadialMenuBuilder<EUnit>::FOnSelectItem::CreateLambda([this](const EUnit& Unit)
+			{
+				Implementation.SetDataUnit(Unit);
+			}), MenuWidth);
 	}
 }
 
diff --git a/Engine/Plugins/Developer/Concert/ConcertApp/MultiUserServer/Source/MultiUserServer/Private/Widgets/Clients/Logging/Filter/ConcertFrontendLogFilter_Size.h b/Engine/Plugins/Developer/Concert/ConcertApp/MultiUserServer/Source/MultiUserServer/Private/Widgets/Clients/Logging/Filter/ConcertFrontendLogFilter_Size.h
--- a/Engine/Plugins/Developer/Concert/ConcertApp/MultiUserServer/Source/MultiUserServer/Private/Widgets/Clients/Logging/Filter/ConcertFrontendLogFilter_Size.h
+++ b/Engine/Plugins/Developer/Concert/ConcertApp/MultiUserServer/Source/MultiUserServer/Private/Widgets/Clients/Logging/Filter/ConcertFrontendLogFilter_Size.h
@@ -4,11 +4,10 @@
 
 #include "CoreMinimal.h"
 
-#include "ConcertLogFilterTypes.h"
-#include "Widgets/Util/Filter/ConcertFilter.h"
 #include "Widgets/Util/Filter/ConcertFrontendFilter.h"
 
 #include "Math/UnitConversion.h"
+#include "Widgets/Clients/Logging/ConcertLogEntry.h"
 
 namespace UE::MultiUserServer
 {
@@ -21,19 +20,20 @@ namespace UE::MultiUserServer
 	};
 
 	/** Filters based on the log's size */
-	class FConcertLogFilter_Size : public FConcertLogFilter
+	class FConcertLogFilter_Size : public IFilter<const FConcertLogEntry&>
 	{
 	public:
 	
 		//~ Begin FConcertLogFilter Interface
 		virtual bool PassesFilter(const FConcertLogEntry& InItem) const override;
+		virtual FChangedEvent& OnChanged() override { return ChangedEvent; }
 		//~ End FConcertLogFilter Interface
 
-		void AdvanceFilterMode();
+		void SetOperator(ESizeFilterMode Operator);
 		void SetSizeInBytes(uint32 NewSizeInBytes);
 		void SetDataUnit(EUnit NewUnit);
 	
-		ESizeFilterMode GetFilterMode() const { return FilterMode; }
+		ESizeFilterMode GetOperator() const { return FilterMode; }
 		uint32 GetSizeInBytes() const { return SizeInBytes; }
 		EUnit GetDataUnit() const { return DataUnit; }
 		TSet<EUnit> GetAllowedUnits() const { return { EUnit::Bytes, EUnit::Kilobytes, EUnit::Megabytes }; }
@@ -43,6 +43,8 @@ namespace UE::MultiUserServer
 		ESizeFilterMode FilterMode = ESizeFilterMode::BiggerThanOrEqual;
 		uint32 SizeInBytes = 0;
 		EUnit DataUnit = EUnit::Bytes;
+
+		FChangedEvent ChangedEvent;
 	};
 
 	class FConcertFrontendLogFilter_Size : public TConcertFrontendFilterAggregate<FConcertLogFilter_Size, const FConcertLogEntry&>
@@ -50,11 +52,11 @@ namespace UE::MultiUserServer
 		using Super = TConcertFrontendFilterAggregate<FConcertLogFilter_Size, const FConcertLogEntry&>;
 	public:
 	
-		FConcertFrontendLogFilter_Size();
+		FConcertFrontendLogFilter_Size(TSharedRef<FFilterCategory> FilterCategory);
 
-	private:
+		virtual FString GetName() const override { return TEXT("Size"); }
+		virtual FText GetDisplayName() const override { return NSLOCTEXT("UnrealMultiUserUI.FConcertFrontendLogFilter_Size", "DisplayLabel", "Size"); }
 
-		TSharedRef<SWidget> MakeDataUnitMenu();
-		FText GetSizeAndUnitAsText() const;
+		virtual void ExposeEditWidgets(FMenuBuilder& MenuBuilder) override;
 	};
 }
diff --git a/Engine/Plugins/Developer/Concert/ConcertApp/MultiUserServer/Source/MultiUserServer/Private/Widgets/Clients/Logging/Filter/ConcertFrontendLogFilter_TextSearch.h b/Engine/Plugins/Developer/Concert/ConcertApp/MultiUserServer/Source/MultiUserServer/Private/Widgets/Clients/Logging/Filter/ConcertFrontendLogFilter_TextSearch.h
--- a/Engine/Plugins/Developer/Concert/ConcertApp/MultiUserServer/Source/MultiUserServer/Private/Widgets/Clients/Logging/Filter/ConcertFrontendLogFilter_TextSearch.h
+++ b/Engine/Plugins/Developer/Concert/ConcertApp/MultiUserServer/Source/MultiUserServer/Private/Widgets/Clients/Logging/Filter/ConcertFrontendLogFilter_TextSearch.h
@@ -41,8 +41,8 @@ namespace UE::MultiUserServer
 		using Super = TConcertFrontendFilter_TextSearch<FConcertLogFilter_TextSearch, const FConcertLogEntry&>;
 	public:
 
-		FConcertFrontendLogFilter_TextSearch(TSharedRef<FConcertLogTokenizer> Tokenizer)
-			: Super(MoveTemp(Tokenizer))
+		FConcertFrontendLogFilter_TextSearch(TSharedPtr<FFilterCategory> InCategory, TSharedRef<FConcertLogTokenizer> Tokenizer)
+			: Super(MoveTemp(InCategory), MoveTemp(Tokenizer))
 		{}
 	};
 }
diff --git a/Engine/Plugins/Developer/Concert/ConcertApp/MultiUserServer/Source/MultiUserServer/Private/Widgets/Clients/Logging/Filter/ConcertFrontendLogFilter_Time.cpp b/Engine/Plugins/Developer/Concert/ConcertApp/MultiUserServer/Source/MultiUserServer/Private/Widgets/Clients/Logging/Filter/ConcertFrontendLogFilter_Time.cpp
--- a/Engine/Plugins/Developer/Concert/ConcertApp/MultiUserServer/Source/MultiUserServer/Private/Widgets/Clients/Logging/Filter/ConcertFrontendLogFilter_Time.cpp
+++ b/Engine/Plugins/Developer/Concert/ConcertApp/MultiUserServer/Source/MultiUserServer/Private/Widgets/Clients/Logging/Filter/ConcertFrontendLogFilter_Time.cpp
@@ -3,13 +3,15 @@
 #include "ConcertFrontendLogFilter_Time.h"
 
 #include "ConcertTransportEvents.h"
+#include "SSimpleComboButton.h"
 #include "Widgets/Clients/Logging/ConcertLogEntry.h"
 
 #include "Framework/MultiBox/MultiBoxBuilder.h"
-#include "Widgets/Input/SComboButton.h"
-#include "Widgets/Input/SEditableTextBox.h"
+#include "Styling/StyleColors.h"
 #include "Widgets/SBoxPanel.h"
+#include "Widgets/Input/SNumericEntryBox.h"
 #include "Widgets/Text/STextBlock.h"
+#include "Widgets/Util/Filter/ConcertFilterUtils.h"
 
 #define LOCTEXT_NAMESPACE "UnrealMultiUserUI.FConcertFrontendLogFilter_Time"
 
@@ -61,47 +63,18 @@ namespace UE::MultiUserServer
 		}
 	}
 
-	FConcertFrontendLogFilter_Time::FConcertFrontendLogFilter_Time(ETimeFilter TimeFilter)
-		: Super(TimeFilter)
-	{
-		ChildSlot = SNew(SHorizontalBox)
-			.ToolTipText(LOCTEXT("TimeFilter.ToolTipText", "Filter logs by local time"))
-
-			+SHorizontalBox::Slot()
-			.AutoWidth()
-			.VAlign(VAlign_Center)
-			[
-				SNew(STextBlock)
-				.Text_Lambda([this]()
-				{
-					return Implementation.GetFilterMode() == ETimeFilter::AllowAfter
-						? LOCTEXT("TimeFilter.AllowAfter", "After")
-						: LOCTEXT("TimeFilter.AllowBefore", "Before");
-				})
-			]
-
-			+SHorizontalBox::Slot()
-			.AutoWidth()
-			.Padding(2, 0, 0, 0)
-			.VAlign(VAlign_Center)
-			[
-				SNew(SComboButton)
-				.OnGetMenuContent_Raw(this, &FConcertFrontendLogFilter_Time::CreateDatePicker)
-				.ButtonContent()
-				[
-					SNew(STextBlock)
-					.Text_Lambda([this](){ return FText::AsDateTime(Implementation.GetTime()); })
-				]
-			];
-	}
+	FConcertFrontendLogFilter_Time::FConcertFrontendLogFilter_Time(TSharedRef<FFilterCategory> FilterCategory, ETimeFilter TimeFilter)
+		: Super(MoveTemp(FilterCategory), TimeFilter)
+		, TimeFilter(TimeFilter)
+	{}
 
-	TSharedRef<SWidget> FConcertFrontendLogFilter_Time::CreateDatePicker()
+	void FConcertFrontendLogFilter_Time::ExposeEditWidgets(FMenuBuilder& MenuBuilder)
 	{
-		FMenuBuilder MenuBuilder(true, nullptr);
-
+		using namespace UE::ConcertFilterUtils;
+		
 		MenuBuilder.AddMenuEntry(
-			LOCTEXT("TimeFilter.Clear.", "Clear"),
-			LOCTEXT("TimeFilter.Clear.Tooltip", "Sets the time so that this filter has no effect"),
+			LOCTEXT("TimeFilter.Reset.", "Reset"),
+			LOCTEXT("TimeFilter.Reset.Tooltip", "Sets the time so that this filter has no effect"),
 			FSlateIcon(),
 			FUIAction(
 				FExecuteAction::CreateLambda([this](){ Implementation.ResetToInfiniteTime(); }),
@@ -121,22 +94,169 @@ namespace UE::MultiUserServer
 			NAME_None,
 			EUserInterfaceActionType::Button
 		);
+		
 		MenuBuilder.AddSeparator();
 
-		MenuBuilder.AddWidget(
-			SNew(SEditableTextBox)
-			.Text(FText::FromString(Implementation.GetTime().ToString()))
-			.OnTextCommitted_Lambda([this](const FText& Text, ETextCommit::Type)
+		// Year
+		MenuBuilder.AddWidget(SetMenuWidgetWidth(CreateDateComponentNumericWidget(ETimeComponent::Year)), LOCTEXT("Year", "Year"), true);
+
+		// Month
+		using FMonth = int32;
+		TRadialMenuBuilder<FMonth>::AddRadialSubMenu(
+			MenuBuilder,
+			LOCTEXT("Month", "Month"),
+			TAttribute<FMonth>::CreateLambda([this](){ return Implementation.GetTime().GetMonth(); }),
+			TArray<FMonth>{ 1, 2, 3, 4, 5, 6, 7, 8, 9, 10, 11, 12 },
+			TRadialMenuBuilder<FMonth>::FItemToText::CreateLambda([](const FMonth& Month)
 			{
-				FDateTime Time;
-				if (FDateTime::Parse(Text.ToString(), Time))
-				{
-					Implementation.SetTime(Time);
-				}
+				const FDateTime Date(1, Month, 1);
+				return FText::FromString(Date.ToFormattedString(TEXT("%B")));
 			}),
-			LOCTEXT("TimeFilter.Custom.", "Custom time")
+			TRadialMenuBuilder<FMonth>::FOnSelectItem::CreateLambda([this](const FMonth& Month)
+			{
+				const FDateTime CurrentTime = Implementation.GetTime();
+				const int32 Day = FMath::Clamp(CurrentTime.GetDay(), 1, FDateTime::DaysInMonth(CurrentTime.GetYear(), Month));
+				const FDateTime NewTime(CurrentTime.GetYear(), Month, Day, CurrentTime.GetHour(), CurrentTime.GetMinute(), CurrentTime.GetSecond());
+				Implementation.SetTime(NewTime);
+			})
 			);
-		return MenuBuilder.MakeWidget();
+
+		// Day
+		MenuBuilder.AddWidget(SetMenuWidgetWidth(CreateDateComponentNumericWidget(ETimeComponent::Day)), LOCTEXT("Day", "Day"), true);
+
+		// Hour : Minute : Second
+		const TSharedRef<SWidget> TimeWidget = SNew(SHorizontalBox)
+			+SHorizontalBox::Slot()
+			.AutoWidth()
+			.VAlign(VAlign_Center)
+			[
+				CreateDateComponentNumericWidget(ETimeComponent::Hour)
+			]
+			+SHorizontalBox::Slot()
+			.FillWidth(1.f)
+			.HAlign(HAlign_Center)
+			.VAlign(VAlign_Center)
+			[
+				SNew(STextBlock)
+				.Text(LOCTEXT("Colon", ":"))
+			]
+			+SHorizontalBox::Slot()
+			.AutoWidth()
+			.VAlign(VAlign_Center)
+			[
+				CreateDateComponentNumericWidget(ETimeComponent::Minute)
+			]
+			+SHorizontalBox::Slot()
+			.FillWidth(1.f)
+			.HAlign(HAlign_Center)
+			.VAlign(VAlign_Center)
+			[
+				SNew(STextBlock)
+				.Text(LOCTEXT("Colon", ":"))
+			]
+			+SHorizontalBox::Slot()
+			.AutoWidth()
+			.VAlign(VAlign_Center)
+			[
+				CreateDateComponentNumericWidget(ETimeComponent::Second)
+			];
+		MenuBuilder.AddWidget(SetMenuWidgetWidth(TimeWidget), LOCTEXT("Time", "Time"), true);
+	}
+
+	void FConcertFrontendLogFilter_Time::BuildMonthSubMenu(FMenuBuilder& MenuBuilder)
+	{
+		for (int32 MonthIdx = 1; MonthIdx <= 12; ++MonthIdx)
+		{
+			const FDateTime Date(1, MonthIdx, 1);
+			MenuBuilder.AddMenuEntry(
+				// B means long month name
+				FText::FromString(Date.ToFormattedString(TEXT("%B"))),
+				FText::GetEmpty(),
+				FSlateIcon(),
+				FUIAction(
+					FExecuteAction::CreateLambda([this, MonthIdx]()
+					{
+						const FDateTime CurrentTime = Implementation.GetTime();
+						const FDateTime NewTime(CurrentTime.GetYear(), MonthIdx, CurrentTime.GetDay(), CurrentTime.GetHour(), CurrentTime.GetMinute(), CurrentTime.GetSecond());
+						Implementation.SetTime(NewTime);
+					}),
+					FCanExecuteAction::CreateLambda([] { return true; }),
+					FIsActionChecked::CreateLambda([this, MonthIdx](){ return Implementation.GetTime().GetMonth() == MonthIdx; })),
+				NAME_None,
+				EUserInterfaceActionType::RadioButton
+			);
+		}
+	}
+
+	TSharedRef<SWidget> FConcertFrontendLogFilter_Time::CreateDateComponentNumericWidget(ETimeComponent TimeComponent)
+	{
+		const uint32 Min = [TimeComponent]()
+		{
+			switch (TimeComponent)
+			{
+			case ETimeComponent::Year: return 1;
+			case ETimeComponent::Day: return 1;
+			case ETimeComponent::Hour: return 0;
+			case ETimeComponent::Minute: return 0;
+			case ETimeComponent::Second: return 0;
+			default: checkNoEntry(); return 0;;
+			}
+		}();
+		const uint32 Max = [TimeComponent]()
+		{
+			switch (TimeComponent)
+			{
+			case ETimeComponent::Year: return 9999;
+			case ETimeComponent::Day: return 31;
+			case ETimeComponent::Hour: return 23;
+			case ETimeComponent::Minute: return 59;
+			case ETimeComponent::Second: return 59;
+			default: checkNoEntry(); return 0;;
+			}
+		}();
+
+		auto MaxLambda = [this, TimeComponent, Max]() -> uint32
+		{
+			if (TimeComponent == ETimeComponent::Day)
+			{
+				const FDateTime Time = Implementation.GetTime();
+				return FDateTime::DaysInMonth(Time.GetYear(), Time.GetMonth());
+			}
+			return Max;
+		};
+		
+		return SNew(SNumericEntryBox<uint32>)
+			.AllowSpin(true)
+			.MinDesiredValueWidth(Min)
+			.MinValue(Min)
+			.MinSliderValue(Min)
+			.MaxSliderValue_Lambda(MaxLambda)
+			.MaxValue_Lambda(MaxLambda)
+			.OnValueChanged_Lambda([this, TimeComponent](uint32 NewValue)
+			{
+				const FDateTime CurrentTime = Implementation.GetTime();
+				switch (TimeComponent)
+				{
+				case ETimeComponent::Year: Implementation.SetTime({ static_cast<int32>(NewValue), CurrentTime.GetMonth(), CurrentTime.GetDay(), CurrentTime.GetHour(), CurrentTime.GetMinute(), CurrentTime.GetSecond() }); break;
+				case ETimeComponent::Day: Implementation.SetTime({ CurrentTime.GetYear(), CurrentTime.GetMonth(), static_cast<int32>(NewValue), CurrentTime.GetHour(), CurrentTime.GetMinute(), CurrentTime.GetSecond() }); break;
+				case ETimeComponent::Hour: Implementation.SetTime({ CurrentTime.GetYear(), CurrentTime.GetMonth(), CurrentTime.GetDay(), static_cast<int32>(NewValue), CurrentTime.GetMinute(), CurrentTime.GetSecond() }); break;
+				case ETimeComponent::Minute: Implementation.SetTime({ CurrentTime.GetYear(), CurrentTime.GetMonth(), CurrentTime.GetDay(), CurrentTime.GetHour(), static_cast<int32>(NewValue), CurrentTime.GetSecond() }); break;
+				case ETimeComponent::Second: Implementation.SetTime({ CurrentTime.GetYear(), CurrentTime.GetMonth(), CurrentTime.GetDay(), CurrentTime.GetHour(), CurrentTime.GetMinute(), static_cast<int32>(NewValue) }); break;
+				default: ;
+				}
+			})
+			.Value_Lambda([this, TimeComponent]()
+			{
+				switch (TimeComponent)
+				{
+				case ETimeComponent::Year: return Implementation.GetTime().GetYear();
+				case ETimeComponent::Day: return Implementation.GetTime().GetDay();
+				case ETimeComponent::Hour: return Implementation.GetTime().GetHour();
+				case ETimeComponent::Minute: return Implementation.GetTime().GetMinute();
+				case ETimeComponent::Second: return Implementation.GetTime().GetSecond();
+				default: return 0;
+				}
+			});
 	}
 }
 
diff --git a/Engine/Plugins/Developer/Concert/ConcertApp/MultiUserServer/Source/MultiUserServer/Private/Widgets/Clients/Logging/Filter/ConcertFrontendLogFilter_Time.h b/Engine/Plugins/Developer/Concert/ConcertApp/MultiUserServer/Source/MultiUserServer/Private/Widgets/Clients/Logging/Filter/ConcertFrontendLogFilter_Time.h
--- a/Engine/Plugins/Developer/Concert/ConcertApp/MultiUserServer/Source/MultiUserServer/Private/Widgets/Clients/Logging/Filter/ConcertFrontendLogFilter_Time.h
+++ b/Engine/Plugins/Developer/Concert/ConcertApp/MultiUserServer/Source/MultiUserServer/Private/Widgets/Clients/Logging/Filter/ConcertFrontendLogFilter_Time.h
@@ -3,9 +3,7 @@
 #pragma once
 
 #include "CoreMinimal.h"
-#include "ConcertLogFilterTypes.h"
-
-#include "Widgets/Util/Filter/ConcertFilter.h"
+#include "Widgets/Clients/Logging/ConcertLogEntry.h"
 #include "Widgets/Util/Filter/ConcertFrontendFilter.h"
 
 namespace UE::MultiUserServer
@@ -19,7 +17,7 @@ namespace UE::MultiUserServer
 	};
 
 	/** Filters based on whether a log happened before or after a certain time */
-	class FConcertLogFilter_Time : public FConcertLogFilter
+	class FConcertLogFilter_Time : public IFilter<const FConcertLogEntry&>
 	{
 	public:
 
@@ -29,6 +27,7 @@ namespace UE::MultiUserServer
 
 		//~ Begin FConcertLogFilter Interface
 		virtual bool PassesFilter(const FConcertLogEntry& InItem) const override;
+		virtual FChangedEvent& OnChanged() override { return ChangedEvent; }
 		//~ End FConcertLogFilter Interfac
 
 		ETimeFilter GetFilterMode() const { return FilterMode; }
@@ -41,6 +40,7 @@ namespace UE::MultiUserServer
 
 		ETimeFilter FilterMode;
 		FDateTime Time;
+		FChangedEvent ChangedEvent;
 
 		FDateTime MakeResetTime() const;
 	};
@@ -50,11 +50,33 @@ namespace UE::MultiUserServer
 		using Super = TConcertFrontendFilterAggregate<FConcertLogFilter_Time, const FConcertLogEntry&>;
 	public:
 	
-		FConcertFrontendLogFilter_Time(ETimeFilter TimeFilter);
+		FConcertFrontendLogFilter_Time(TSharedRef<FFilterCategory> FilterCategory, ETimeFilter TimeFilter);
+		
+		virtual FString GetName() const override { return TimeFilter == ETimeFilter::AllowAfter ? TEXT("After") : TEXT("Before"); }
+		virtual FText GetDisplayName() const override
+		{
+			return TimeFilter == ETimeFilter::AllowAfter
+				? NSLOCTEXT("UnrealMultiUserUI.FConcertFrontendLogFilter_Time", "DisplayLabel.After", "After Timestamp")
+				: NSLOCTEXT("UnrealMultiUserUI.FConcertFrontendLogFilter_Time", "DisplayLabel.Before", "Before Timestamp");
+		}
+
+		virtual void ExposeEditWidgets(FMenuBuilder& MenuBuilder) override;
 
 	private:
 
-		TSharedRef<SWidget> CreateDatePicker();
+		enum class ETimeComponent
+		{
+			Year,
+			Day,
+			Hour,
+			Minute,
+			Second
+		};
+		
+		ETimeFilter TimeFilter;
+		
+		void BuildMonthSubMenu(FMenuBuilder& MenuBuilder);
+		TSharedRef<SWidget> CreateDateComponentNumericWidget(ETimeComponent TimeComponent);
 	};
 }
 
diff --git a/Engine/Plugins/Developer/Concert/ConcertApp/MultiUserServer/Source/MultiUserServer/Private/Widgets/Clients/Logging/Filter/ConcertLogFilterTypes.h b/Engine/Plugins/Developer/Concert/ConcertApp/MultiUserServer/Source/MultiUserServer/Private/Widgets/Clients/Logging/Filter/ConcertLogFilterTypes.h
deleted file mode 100644
diff --git a/Engine/Plugins/Developer/Concert/ConcertApp/MultiUserServer/Source/MultiUserServer/Private/Widgets/Clients/Logging/Filter/ConcertLogFilter_FrontendRoot.cpp b/Engine/Plugins/Developer/Concert/ConcertApp/MultiUserServer/Source/MultiUserServer/Private/Widgets/Clients/Logging/Filter/ConcertLogFilter_FrontendRoot.cpp
--- a/Engine/Plugins/Developer/Concert/ConcertApp/MultiUserServer/Source/MultiUserServer/Private/Widgets/Clients/Logging/Filter/ConcertLogFilter_FrontendRoot.cpp
+++ b/Engine/Plugins/Developer/Concert/ConcertApp/MultiUserServer/Source/MultiUserServer/Private/Widgets/Clients/Logging/Filter/ConcertLogFilter_FrontendRoot.cpp
@@ -9,43 +9,52 @@
 #include "ConcertFrontendLogFilter_TextSearch.h"
 #include "ConcertFrontendLogFilter_Time.h"
 #include "ConcertFrontendLogFilter_Size.h"
-#include "ConcertLogFilterTypes.h"
 
 
 namespace UE::MultiUserServer
 {
 	namespace Private
 	{
-		static TArray<TSharedRef<FConcertFrontendLogFilter>> CreateCommonFilters()
+		static TArray<TSharedRef<TConcertFrontendFilter<const FConcertLogEntry&>>> CreateCommonFilters(const TSharedRef<FFilterCategory>& FilterCategory)
 		{
 			return {
-				MakeShared<FConcertFrontendLogFilter_MessageAction>(),
-				MakeShared<FConcertFrontendLogFilter_MessageType>(),
-				MakeShared<FConcertFrontendLogFilter_Time>(ETimeFilter::AllowAfter),
-				MakeShared<FConcertFrontendLogFilter_Time>(ETimeFilter::AllowBefore),
-				MakeShared<FConcertFrontendLogFilter_Size>(),
-				MakeShared<FConcertFrontendLogFilter_Ack>()
+				MakeShared<FConcertFrontendLogFilter_MessageAction>(FilterCategory),
+				MakeShared<FConcertFrontendLogFilter_MessageType>(FilterCategory),
+				MakeShared<FConcertFrontendLogFilter_Time>(FilterCategory, ETimeFilter::AllowAfter),
+				MakeShared<FConcertFrontendLogFilter_Time>(FilterCategory, ETimeFilter::AllowBefore),
+				MakeShared<FConcertFrontendLogFilter_Size>(FilterCategory),
+				MakeShared<FConcertFrontendLogFilter_Ack>(FilterCategory)
 			};
 		}
 	}
 	
 	TSharedRef<FConcertLogFilter_FrontendRoot> MakeGlobalLogFilter(TSharedRef<FConcertLogTokenizer> Tokenizer)
 	{
+		TSharedRef<FFilterCategory> FilterCategory = MakeShared<FFilterCategory>(
+			NSLOCTEXT("UnrealMultiUserUI.TConcertFrontendRootFilter", "DefaultCategoryLabel", "Default"),
+			FText::GetEmpty()
+			);
 		return MakeShared<FConcertLogFilter_FrontendRoot>(
 			MoveTemp(Tokenizer),
-			Private::CreateCommonFilters()
+			Private::CreateCommonFilters(FilterCategory),
+			FilterCategory
 			);
 	}
 
 	TSharedRef<FConcertLogFilter_FrontendRoot> MakeClientLogFilter(TSharedRef<FConcertLogTokenizer> Tokenizer, const FGuid& ClientMessageNodeId, const TSharedRef<FEndpointToUserNameCache>& EndpointCache)
 	{
-		const TArray<TSharedRef<FConcertFrontendLogFilter>> CommonFilters = Private::CreateCommonFilters();
-		const TArray<TSharedRef<FConcertLogFilter>> NonVisuals = {
+		TSharedRef<FFilterCategory> FilterCategory = MakeShared<FFilterCategory>(
+			NSLOCTEXT("UnrealMultiUserUI.TConcertFrontendRootFilter", "DefaultCategoryLabel", "Default"),
+			FText::GetEmpty()
+			);
+		const TArray<TSharedRef<TConcertFrontendFilter<const FConcertLogEntry&>>> CommonFilters = Private::CreateCommonFilters(FilterCategory);
+		const TArray<TSharedRef<IFilter<const FConcertLogEntry&>>> NonVisuals = {
 			MakeShared<FConcertLogFilter_Client>(ClientMessageNodeId, EndpointCache)
 		};
 		return MakeShared<FConcertLogFilter_FrontendRoot>(
 			MoveTemp(Tokenizer),
 			CommonFilters,
+			MoveTemp(FilterCategory),
 			NonVisuals
 			);
 	}
diff --git a/Engine/Plugins/Developer/Concert/ConcertApp/MultiUserServer/Source/MultiUserServer/Private/Widgets/Clients/Logging/Filter/ConcertLogFilter_FrontendRoot.h b/Engine/Plugins/Developer/Concert/ConcertApp/MultiUserServer/Source/MultiUserServer/Private/Widgets/Clients/Logging/Filter/ConcertLogFilter_FrontendRoot.h
--- a/Engine/Plugins/Developer/Concert/ConcertApp/MultiUserServer/Source/MultiUserServer/Private/Widgets/Clients/Logging/Filter/ConcertLogFilter_FrontendRoot.h
+++ b/Engine/Plugins/Developer/Concert/ConcertApp/MultiUserServer/Source/MultiUserServer/Private/Widgets/Clients/Logging/Filter/ConcertLogFilter_FrontendRoot.h
@@ -20,10 +20,11 @@ namespace UE::MultiUserServer
 
 		FConcertLogFilter_FrontendRoot(
 			TSharedRef<FConcertLogTokenizer> Tokenizer,
-			TArray<TSharedRef<Super::FConcertFrontendFilter>> FrontendFilters,
-			const TArray<TSharedRef<Super::FConcertFilter>>& NonVisualFilters = {}
+			TArray<TSharedRef<TConcertFrontendFilter<const FConcertLogEntry&>>> FrontendFilters,
+			TSharedRef<FFilterCategory> DefaultCategory,
+			const TArray<TSharedRef<IFilter<const FConcertLogEntry&>>>& NonVisualFilters = {}
 			)
-			: Super(MakeShared<FConcertFrontendLogFilter_TextSearch>(MoveTemp(Tokenizer)), MoveTemp(FrontendFilters), NonVisualFilters)
+			: Super(MakeShared<FConcertFrontendLogFilter_TextSearch>(DefaultCategory, MoveTemp(Tokenizer)), MoveTemp(FrontendFilters), NonVisualFilters, DefaultCategory)
 		{}
 	};
 	
diff --git a/Engine/Plugins/Developer/Concert/ConcertApp/MultiUserServer/Source/MultiUserServer/Private/Widgets/Clients/Logging/Filter/FilteredConcertLogList.cpp b/Engine/Plugins/Developer/Concert/ConcertApp/MultiUserServer/Source/MultiUserServer/Private/Widgets/Clients/Logging/Filter/FilteredConcertLogList.cpp
--- a/Engine/Plugins/Developer/Concert/ConcertApp/MultiUserServer/Source/MultiUserServer/Private/Widgets/Clients/Logging/Filter/FilteredConcertLogList.cpp
+++ b/Engine/Plugins/Developer/Concert/ConcertApp/MultiUserServer/Source/MultiUserServer/Private/Widgets/Clients/Logging/Filter/FilteredConcertLogList.cpp
@@ -2,11 +2,10 @@
 
 #include "FilteredConcertLogList.h"
 
-#include "Widgets/Util/Filter/ConcertFilter.h"
 
 namespace UE::MultiUserServer
 {
-	FFilteredConcertLogList::FFilteredConcertLogList(TSharedRef<IConcertLogSource> InLogSource, TSharedPtr<FConcertLogFilter> InOptionalFilter)
+	FFilteredConcertLogList::FFilteredConcertLogList(TSharedRef<IConcertLogSource> InLogSource, TSharedPtr<IFilter<const FConcertLogEntry&>> InOptionalFilter)
 		: LogSource(MoveTemp(InLogSource))
 		, Filter(MoveTemp(InOptionalFilter))
 	{
@@ -91,7 +90,7 @@ namespace UE::MultiUserServer
 		}
 	}
 
-	FPagedFilteredConcertLogList::FPagedFilteredConcertLogList(TSharedRef<IConcertLogSource> LogSource, TSharedPtr<FConcertLogFilter> OptionalFilter, uint16 InitialLogsPerPage)
+	FPagedFilteredConcertLogList::FPagedFilteredConcertLogList(TSharedRef<IConcertLogSource> LogSource, TSharedPtr<IFilter<const FConcertLogEntry&>> OptionalFilter, uint16 InitialLogsPerPage)
 		: FFilteredConcertLogList(MoveTemp(LogSource), MoveTemp(OptionalFilter))
 		, LogsPerPage(FMath::Max<uint16>(InitialLogsPerPage, 1))
 	{
diff --git a/Engine/Plugins/Developer/Concert/ConcertApp/MultiUserServer/Source/MultiUserServer/Private/Widgets/Clients/Logging/Filter/FilteredConcertLogList.h b/Engine/Plugins/Developer/Concert/ConcertApp/MultiUserServer/Source/MultiUserServer/Private/Widgets/Clients/Logging/Filter/FilteredConcertLogList.h
--- a/Engine/Plugins/Developer/Concert/ConcertApp/MultiUserServer/Source/MultiUserServer/Private/Widgets/Clients/Logging/Filter/FilteredConcertLogList.h
+++ b/Engine/Plugins/Developer/Concert/ConcertApp/MultiUserServer/Source/MultiUserServer/Private/Widgets/Clients/Logging/Filter/FilteredConcertLogList.h
@@ -3,7 +3,7 @@
 #pragma once
 
 #include "CoreMinimal.h"
-#include "ConcertLogFilterTypes.h"
+#include "Misc/IFilter.h"
 #include "Widgets/Clients/Logging/Source/IConcertLogSource.h"
 
 struct FConcertLog;
@@ -15,7 +15,7 @@ namespace UE::MultiUserServer
 	{
 	public:
 		
-		FFilteredConcertLogList(TSharedRef<IConcertLogSource> InLogSource, TSharedPtr<FConcertLogFilter> InOptionalFilter);
+		FFilteredConcertLogList(TSharedRef<IConcertLogSource> InLogSource, TSharedPtr<IFilter<const FConcertLogEntry&>> InOptionalFilter);
 		~FFilteredConcertLogList();
 		
 		DECLARE_MULTICAST_DELEGATE_OneParam(FOnLogListChanged, const TArray<TSharedPtr<FConcertLogEntry>>& /*NewFilteredView*/);
@@ -39,7 +39,7 @@ namespace UE::MultiUserServer
 		TSharedRef<IConcertLogSource> LogSource;
 
 		/** Optional. Filters items from the LogSource */
-		TSharedPtr<FConcertLogFilter> Filter;
+		TSharedPtr<IFilter<const FConcertLogEntry&>> Filter;
 
 		/** Contains the items of LogSource with Filter applied. */
 		TArray<TSharedPtr<FConcertLogEntry>> FilteredLogs;
@@ -65,7 +65,7 @@ namespace UE::MultiUserServer
 		using FLogsPerPageCount = uint16;
 		using FPageCount = uint32;
 
-		FPagedFilteredConcertLogList(TSharedRef<IConcertLogSource> LogSource, TSharedPtr<FConcertLogFilter> OptionalFilter, uint16 FLogsPerPageCount = 500);
+		FPagedFilteredConcertLogList(TSharedRef<IConcertLogSource> LogSource, TSharedPtr<IFilter<const FConcertLogEntry&>> OptionalFilter, uint16 FLogsPerPageCount = 500);
 
 		FPageCount GetNumPages() const { const FPageCount PageCount = GetFilteredLogs().Num() / LogsPerPage; return GetFilteredLogs().Num() % LogsPerPage == 0 ? PageCount : PageCount + 1;}
 		FPageCount GetCurrentPage() const { return CurrentPageIndex; }
diff --git a/Engine/Plugins/Developer/Concert/ConcertApp/MultiUserServer/Source/MultiUserServer/Private/Widgets/Clients/PackageTransmission/Filter/PackageTransmissionFilterTypes.h b/Engine/Plugins/Developer/Concert/ConcertApp/MultiUserServer/Source/MultiUserServer/Private/Widgets/Clients/PackageTransmission/Filter/PackageTransmissionFilterTypes.h
deleted file mode 100644
diff --git a/Engine/Plugins/Developer/Concert/ConcertApp/MultiUserServer/Source/MultiUserServer/Private/Widgets/Clients/PackageTransmission/Filter/PackageTransmissionFilter_FrontendRoot.h b/Engine/Plugins/Developer/Concert/ConcertApp/MultiUserServer/Source/MultiUserServer/Private/Widgets/Clients/PackageTransmission/Filter/PackageTransmissionFilter_FrontendRoot.h
--- a/Engine/Plugins/Developer/Concert/ConcertApp/MultiUserServer/Source/MultiUserServer/Private/Widgets/Clients/PackageTransmission/Filter/PackageTransmissionFilter_FrontendRoot.h
+++ b/Engine/Plugins/Developer/Concert/ConcertApp/MultiUserServer/Source/MultiUserServer/Private/Widgets/Clients/PackageTransmission/Filter/PackageTransmissionFilter_FrontendRoot.h
@@ -4,7 +4,6 @@
 
 #include "CoreMinimal.h"
 #include "PackageTransmissionFrontendFilter_TextSearch.h"
-#include "PackageTransmissionFilterTypes.h"
 #include "Widgets/Util/Filter/ConcertRootFrontendFilter.h"
 
 class FEndpointToUserNameCache;
@@ -23,10 +22,11 @@ namespace UE::MultiUserServer
 
 		FPackageTransmissionFilter_FrontendRoot(
 			TSharedRef<FPackageTransmissionEntryTokenizer> Tokenizer,
-			TArray<TSharedRef<Super::FConcertFrontendFilter>> FrontendFilters,
-			const TArray<TSharedRef<Super::FConcertFilter>>& NonVisualFilters = {}
+			TArray<TSharedRef<TConcertFrontendFilter<const FPackageTransmissionEntry&>>> FrontendFilters,
+			const TArray<TSharedRef<IFilter<const FPackageTransmissionEntry&>>>& NonVisualFilters = {},
+			TSharedRef<FFilterCategory> DefaultCategory = MakeShared<FFilterCategory>(NSLOCTEXT("UnrealMultiUserServer.TConcertFrontendRootFilter", "DefaultCategoryLabel", "Default"), FText::GetEmpty())
 			)
-			: Super(MakeShared<FPackageTransmissionFrontendFilter_TextSearch>(MoveTemp(Tokenizer)), MoveTemp(FrontendFilters), NonVisualFilters)
+			: Super(MakeShared<FPackageTransmissionFrontendFilter_TextSearch>(DefaultCategory, MoveTemp(Tokenizer)), MoveTemp(FrontendFilters), NonVisualFilters, DefaultCategory)
 		{}
 	};
 
diff --git a/Engine/Plugins/Developer/Concert/ConcertApp/MultiUserServer/Source/MultiUserServer/Private/Widgets/Clients/PackageTransmission/Filter/PackageTransmissionFrontendFilter_TextSearch.h b/Engine/Plugins/Developer/Concert/ConcertApp/MultiUserServer/Source/MultiUserServer/Private/Widgets/Clients/PackageTransmission/Filter/PackageTransmissionFrontendFilter_TextSearch.h
--- a/Engine/Plugins/Developer/Concert/ConcertApp/MultiUserServer/Source/MultiUserServer/Private/Widgets/Clients/PackageTransmission/Filter/PackageTransmissionFrontendFilter_TextSearch.h
+++ b/Engine/Plugins/Developer/Concert/ConcertApp/MultiUserServer/Source/MultiUserServer/Private/Widgets/Clients/PackageTransmission/Filter/PackageTransmissionFrontendFilter_TextSearch.h
@@ -3,7 +3,6 @@
 #pragma once
 
 #include "CoreMinimal.h"
-#include "PackageTransmissionFilterTypes.h"
 #include "Widgets/Clients/PackageTransmission/Util/PackageTransmissionEntryTokenizer.h"
 #include "Widgets/Util/Filter/ConcertFrontendFilter_TextSearch.h"
 
@@ -41,8 +40,8 @@ namespace UE::MultiUserServer
 		using Super = TConcertFrontendFilter_TextSearch<FPackageTransmissionFilter_TextSearch, const FPackageTransmissionEntry&>;
     public:
     	
-    	FPackageTransmissionFrontendFilter_TextSearch(TSharedRef<FPackageTransmissionEntryTokenizer> Tokenizer)
-    		: Super(MoveTemp(Tokenizer))
+    	FPackageTransmissionFrontendFilter_TextSearch(TSharedPtr<FFilterCategory> InCategory, TSharedRef<FPackageTransmissionEntryTokenizer> Tokenizer)
+    		: Super(MoveTemp(InCategory), MoveTemp(Tokenizer))
     	{}
     };
 }
diff --git a/Engine/Plugins/Developer/Concert/ConcertApp/MultiUserServer/Source/MultiUserServer/Private/Widgets/Clients/PackageTransmission/Model/FilteredPackageTransmissionModel.cpp b/Engine/Plugins/Developer/Concert/ConcertApp/MultiUserServer/Source/MultiUserServer/Private/Widgets/Clients/PackageTransmission/Model/FilteredPackageTransmissionModel.cpp
--- a/Engine/Plugins/Developer/Concert/ConcertApp/MultiUserServer/Source/MultiUserServer/Private/Widgets/Clients/PackageTransmission/Model/FilteredPackageTransmissionModel.cpp
+++ b/Engine/Plugins/Developer/Concert/ConcertApp/MultiUserServer/Source/MultiUserServer/Private/Widgets/Clients/PackageTransmission/Model/FilteredPackageTransmissionModel.cpp
@@ -2,11 +2,9 @@
 
 #include "FilteredPackageTransmissionModel.h"
 
-#include "Widgets/Util/Filter/ConcertFilter.h"
-
 namespace UE::MultiUserServer
 {
-	FFilteredPackageTransmissionModel::FFilteredPackageTransmissionModel(TSharedRef<IPackageTransmissionEntrySource> RealSource, TSharedRef<FPackageTransmissionFilter> Filter)
+	FFilteredPackageTransmissionModel::FFilteredPackageTransmissionModel(TSharedRef<IPackageTransmissionEntrySource> RealSource, TSharedRef<IFilter<const FPackageTransmissionEntry&>> Filter)
 		: RealSource(MoveTemp(RealSource))
 		, Filter(MoveTemp(Filter))
 	{
diff --git a/Engine/Plugins/Developer/Concert/ConcertApp/MultiUserServer/Source/MultiUserServer/Private/Widgets/Clients/PackageTransmission/Model/FilteredPackageTransmissionModel.h b/Engine/Plugins/Developer/Concert/ConcertApp/MultiUserServer/Source/MultiUserServer/Private/Widgets/Clients/PackageTransmission/Model/FilteredPackageTransmissionModel.h
--- a/Engine/Plugins/Developer/Concert/ConcertApp/MultiUserServer/Source/MultiUserServer/Private/Widgets/Clients/PackageTransmission/Model/FilteredPackageTransmissionModel.h
+++ b/Engine/Plugins/Developer/Concert/ConcertApp/MultiUserServer/Source/MultiUserServer/Private/Widgets/Clients/PackageTransmission/Model/FilteredPackageTransmissionModel.h
@@ -4,7 +4,7 @@
 
 #include "CoreMinimal.h"
 #include "IPackageTransmissionEntrySource.h"
-#include "Widgets/Clients/PackageTransmission/Filter/PackageTransmissionFilterTypes.h"
+#include "Misc/IFilter.h"
 
 namespace UE::MultiUserServer
 {
@@ -12,7 +12,7 @@ namespace UE::MultiUserServer
 	{
 	public:
 
-		FFilteredPackageTransmissionModel(TSharedRef<IPackageTransmissionEntrySource> RealSource, TSharedRef<FPackageTransmissionFilter> Filter);
+		FFilteredPackageTransmissionModel(TSharedRef<IPackageTransmissionEntrySource> RealSource, TSharedRef<IFilter<const FPackageTransmissionEntry&>> Filter);
 		virtual ~FFilteredPackageTransmissionModel() override;
 
 		//~ Begin IPackageTransmissionEntrySource Interface
@@ -25,7 +25,7 @@ namespace UE::MultiUserServer
 	private:
 
 		TSharedRef<IPackageTransmissionEntrySource> RealSource;
-		TSharedRef<FPackageTransmissionFilter> Filter;
+		TSharedRef<IFilter<const FPackageTransmissionEntry&>> Filter;
 
 		FOnPackageEntriesAdded OnPackageEntriesAddedDelegate;
 		FOnPackageEntriesModified OnPackageEntriesModifiedDelegate;
diff --git a/Engine/Plugins/Developer/Concert/ConcertApp/MultiUserServer/Source/MultiUserServer/Private/Widgets/Util/Filter/ConcertFilter.h b/Engine/Plugins/Developer/Concert/ConcertApp/MultiUserServer/Source/MultiUserServer/Private/Widgets/Util/Filter/ConcertFilter.h
deleted file mode 100644
diff --git a/Engine/Plugins/Developer/Concert/ConcertApp/MultiUserServer/Source/MultiUserServer/Private/Widgets/Util/Filter/ConcertFilterUtils.cpp b/Engine/Plugins/Developer/Concert/ConcertApp/MultiUserServer/Source/MultiUserServer/Private/Widgets/Util/Filter/ConcertFilterUtils.cpp
new file mode 100644
--- /dev/null
+++ b/Engine/Plugins/Developer/Concert/ConcertApp/MultiUserServer/Source/MultiUserServer/Private/Widgets/Util/Filter/ConcertFilterUtils.cpp
@@ -0,0 +1,33 @@
+// Copyright Epic Games, Inc. All Rights Reserved.
+
+#include "ConcertFilterUtils.h"
+
+#include "Framework/MultiBox/MultiBoxBuilder.h"
+#include "Widgets/Layout/SSpacer.h"
+#include "Widgets/SBoxPanel.h"
+#include "Widgets/Text/STextBlock.h"
+
+namespace UE::ConcertFilterUtils
+{
+	/** Used for widgets you add using FMenuBuilder::AddSubMenu or FMenuBuilder::AddWidget. This makes sure that the widgets have the same width (for consistent visuals). */
+	TSharedRef<SWidget> SetMenuWidgetWidth(const TSharedRef<SWidget>& Widget, bool bNeedsPaddingFromSubmenuRightArrow, const float ForcedWidth, const float SubMenuIconRightPadding)
+	{
+		return SNew(SHorizontalBox)
+			+SHorizontalBox::Slot()
+			.FillWidth(1.f)
+			[
+				SNew(SSpacer)
+			]
+
+			+SHorizontalBox::Slot()
+			.AutoWidth()
+			.Padding(0.f, 0.f, static_cast<int32>(bNeedsPaddingFromSubmenuRightArrow) * SubMenuIconRightPadding, 0.f)
+			[
+				SNew(SBox)
+				.WidthOverride(ForcedWidth)
+				[
+					Widget
+				]
+			];
+	}
+}
\ No newline at end of file
diff --git a/Engine/Plugins/Developer/Concert/ConcertApp/MultiUserServer/Source/MultiUserServer/Private/Widgets/Util/Filter/ConcertFilterUtils.h b/Engine/Plugins/Developer/Concert/ConcertApp/MultiUserServer/Source/MultiUserServer/Private/Widgets/Util/Filter/ConcertFilterUtils.h
new file mode 100644
--- /dev/null
+++ b/Engine/Plugins/Developer/Concert/ConcertApp/MultiUserServer/Source/MultiUserServer/Private/Widgets/Util/Filter/ConcertFilterUtils.h
@@ -0,0 +1,89 @@
+// Copyright Epic Games, Inc. All Rights Reserved.
+
+#pragma once
+
+#include "CoreMinimal.h"
+#include "Framework/MultiBox/MultiBoxBuilder.h"
+#include "Widgets/Input/SComboButton.h"
+
+namespace UE::ConcertFilterUtils
+{
+	/** Used for widgets you add using FMenuBuilder::AddSubMenu or FMenuBuilder::AddWidget. This makes sure that the widgets have the same width (for consistent visuals). */
+	TSharedRef<SWidget> SetMenuWidgetWidth(const TSharedRef<SWidget>& Widget, bool bNeedsPaddingFromSubmenuRightArrow = true, const float ForcedWidth = 100.f, const float SubMenuIconRightPadding = 15.f);
+
+	/** @param TSelectableItem Must be (efficiently) copyable */
+	template<typename TSelectableItem>
+	struct TRadialMenuBuilder
+	{
+		DECLARE_DELEGATE_RetVal_OneParam(FText, FItemToText, const TSelectableItem& Item);
+		DECLARE_DELEGATE_OneParam(FOnSelectItem, const TSelectableItem& Item);
+
+		/**
+		 * Adds a menu entry from which the user can choose an item from a list.
+		 * 
+		 * This function acts as abstraction: UX is evolving and there are two options considered 1. (What we're doing now) create a sub menu. 2. Create a combo button widget.
+		 * Regardless of what we choose in the future, this function abstracts the above so only this function needs to be adjusted if we change the UX - instead of changing the call sites.
+		 */
+		static void AddRadialSubMenu(
+			FMenuBuilder& MenuBuilder,
+			TAttribute<FText> LabelAttribute,
+			TAttribute<TSelectableItem> SelectedItemAttribute,
+			TAttribute<TArray<TSelectableItem>> ItemsAttribute,
+			FItemToText DisplayNameDelegate,
+			FOnSelectItem OnSelectItemDelegate,
+			const float ForcedWidth = 100.f)
+		{
+			const TSharedRef<SWidget> Widget = SNew(SHorizontalBox)
+				+SHorizontalBox::Slot()
+				.AutoWidth()
+				.VAlign(VAlign_Center)
+				.HAlign(HAlign_Center)
+				[
+					SNew(STextBlock)
+					.Text(MoveTemp(LabelAttribute))
+				]
+				+SHorizontalBox::Slot()
+				.FillWidth(1.f)
+				[
+					SetMenuWidgetWidth
+					(
+						SNew(SComboButton)
+						.HasDownArrow(false)
+						.ContentPadding(FMargin(0.f, 2.f)) // So it is as high as other widgets
+						.ButtonContent()
+						[
+							SNew(STextBlock)
+							.Text_Lambda([SelectedItemAttribute, DisplayNameDelegate](){ return DisplayNameDelegate.Execute(SelectedItemAttribute.Get()); })
+						]
+					, false, ForcedWidth)
+				];
+		
+			MenuBuilder.AddSubMenu(
+				Widget,
+				FNewMenuDelegate::CreateLambda(
+				[ItemsAttribute, SelectedItemAttribute, DisplayNameDelegate, OnSelectItemDelegate](FMenuBuilder& MenuBuilder)
+				{
+					for (const TSelectableItem& Item : ItemsAttribute.Get())
+					{
+						MenuBuilder.AddMenuEntry(
+							DisplayNameDelegate.Execute(Item),
+							FText::GetEmpty(),
+							FSlateIcon(),
+							FUIAction(
+								FExecuteAction::CreateLambda([OnSelectItemDelegate, Item](){
+									OnSelectItemDelegate.Execute(Item);
+								}),
+								FCanExecuteAction::CreateLambda([] { return true; }),
+								FIsActionChecked::CreateLambda([SelectedItemAttribute, Item](){ return Item == SelectedItemAttribute.Get(); })),
+							NAME_None,
+							EUserInterfaceActionType::RadioButton
+						);
+					}
+				}),
+				false,
+				false
+				);
+		}
+	};
+	
+}
diff --git a/Engine/Plugins/Developer/Concert/ConcertApp/MultiUserServer/Source/MultiUserServer/Private/Widgets/Util/Filter/ConcertFrontendFilter.h b/Engine/Plugins/Developer/Concert/ConcertApp/MultiUserServer/Source/MultiUserServer/Private/Widgets/Util/Filter/ConcertFrontendFilter.h
--- a/Engine/Plugins/Developer/Concert/ConcertApp/MultiUserServer/Source/MultiUserServer/Private/Widgets/Util/Filter/ConcertFrontendFilter.h
+++ b/Engine/Plugins/Developer/Concert/ConcertApp/MultiUserServer/Source/MultiUserServer/Private/Widgets/Util/Filter/ConcertFrontendFilter.h
@@ -3,8 +3,10 @@
 #pragma once
 
 #include "CoreMinimal.h"
-#include "ConcertFilter.h"
+#include "Framework/MultiBox/MultiBoxBuilder.h"
+#include "Filters/FilterBase.h"
 
+class FMenuBuilder;
 class SWidget;
 
 namespace UE::MultiUserServer
@@ -13,13 +15,31 @@ namespace UE::MultiUserServer
 	template<typename TFilterType>
 	class TConcertFrontendFilter
 		:
-		public TConcertFilter<TFilterType>,
+		public FFilterBase<TFilterType>,
 		public TSharedFromThis<TConcertFrontendFilter<TFilterType>>
 	{
 	public:
 
-		/** Gets a widget that represents this filter */
-		virtual TSharedRef<SWidget> GetFilterWidget() = 0;
+		TConcertFrontendFilter(TSharedPtr<FFilterCategory> InCategory)
+			: FFilterBase<TFilterType>(MoveTemp(InCategory))
+		{}
+
+		/** Exposes widgets with which the filter can be edited. */
+		virtual void ExposeEditWidgets(FMenuBuilder& MenuBuilder) {}
+		
+		virtual void ModifyContextMenu(FMenuBuilder& MenuBuilder)
+		{
+			MenuBuilder.BeginSection("EditFilter", NSLOCTEXT("TConcertFrontendFilter", "EditFilter", "Edit Filter"));
+			ExposeEditWidgets(MenuBuilder);
+		}
+
+		virtual FText GetToolTipText() const { return FText::GetEmpty(); }
+		virtual FLinearColor GetColor() const { return FLinearColor::White; }
+		virtual FName GetIconName() const { return NAME_None; }
+		virtual bool IsInverseFilter() const { return false; }
+		virtual void ActiveStateChanged(bool bActive) {}
+		virtual void SaveSettings(const FString& IniFilename, const FString& IniSection, const FString& SettingsString) const {}
+		virtual void LoadSettings(const FString& IniFilename, const FString& IniSection, const FString& SettingsString) {}
 	};
 
 	/**
@@ -39,24 +59,20 @@ namespace UE::MultiUserServer
 		virtual bool PassesFilter(TFilterType InItem) const final override { return Implementation.PassesFilter(InItem); }
 		//~ End FConcertLogFilter Interface
 	
-		//~ Begin FConcertFrontendLogFilter Interface
-		virtual TSharedRef<SWidget> GetFilterWidget() final override { return ChildSlot.ToSharedRef(); }
-		//~ End FConcertLogFilter Interface
-	
 	protected:
 
 		template<typename... TArg>
-		TConcertFrontendFilterAggregate(TArg&&... Arg)
-			: Implementation(Forward<TArg>(Arg)...)
+		TConcertFrontendFilterAggregate(TSharedPtr<FFilterCategory> InCategory, TArg&&... Arg)
+			: Super(MoveTemp(InCategory))
+			, Implementation(Forward<TArg>(Arg)...)
 		{
 			Implementation.OnChanged().AddLambda([this]()
 			{
-				Super::OnChanged().Broadcast();
+				Super::BroadcastChangedEvent();
 			});
 		}
 
 		TConcertFilterImpl Implementation;
-		TSharedPtr<TWidgetType> ChildSlot;
 	};
 }
 
diff --git a/Engine/Plugins/Developer/Concert/ConcertApp/MultiUserServer/Source/MultiUserServer/Private/Widgets/Util/Filter/ConcertFrontendFilter_TextSearch.h b/Engine/Plugins/Developer/Concert/ConcertApp/MultiUserServer/Source/MultiUserServer/Private/Widgets/Util/Filter/ConcertFrontendFilter_TextSearch.h
--- a/Engine/Plugins/Developer/Concert/ConcertApp/MultiUserServer/Source/MultiUserServer/Private/Widgets/Util/Filter/ConcertFrontendFilter_TextSearch.h
+++ b/Engine/Plugins/Developer/Concert/ConcertApp/MultiUserServer/Source/MultiUserServer/Private/Widgets/Util/Filter/ConcertFrontendFilter_TextSearch.h
@@ -3,7 +3,9 @@
 #pragma once
 
 #include "CoreMinimal.h"
+#include "Misc/IFilter.h"
 #include "ConcertFrontendFilter.h"
+#include "Filters/SFilterSearchBox.h"
 #include "Misc/TextFilter.h"
 #include "Widgets/Input/SSearchBox.h"
 
@@ -13,25 +15,31 @@ namespace UE::MultiUserServer
 {
 	/** Allows advanced search by text. Implements Adapter pattern to wrap TTextFilter. */
 	template<typename TFilterType>
-	class TConcertFilter_TextSearch : public TConcertFilter<TFilterType>
+	class TConcertFilter_TextSearch : public IFilter<TFilterType>
 	{
-		using Super = TConcertFilter<TFilterType>;
 	public:
 
 		TConcertFilter_TextSearch()
-			: TextFilter(TTextFilter<TFilterType>::FItemToStringArray::CreateRaw(this, &TConcertFilter_TextSearch<TFilterType>::GenerateSearchTerms))
+			: TextFilter(MakeTextFilter())
 		{
-			TextFilter.OnChanged().AddLambda([this]
+			TextFilter->OnChanged().AddLambda([this]
 			{
-				Super::OnChanged().Broadcast();
+				OnChangedEvent.Broadcast();
 			});
 		}
-	
+
+		TSharedRef<TTextFilter<TFilterType>> MakeTextFilter() const
+		{
+			return MakeShared<TTextFilter<TFilterType>>(TTextFilter<TFilterType>::FItemToStringArray::CreateRaw(this, &TConcertFilter_TextSearch<TFilterType>::GenerateSearchTerms));
+		}
+		
 		//~ Begin FConcertLogFilter Interface
-		virtual bool PassesFilter(TFilterType InItem) const override { return TextFilter.PassesFilter(InItem); }
+		virtual bool PassesFilter(TFilterType InItem) const override { return TextFilter->PassesFilter(InItem); }
+		virtual typename IFilter<TFilterType>::FChangedEvent& OnChanged() override { return OnChangedEvent; }
 		//~ End FConcertLogFilter Interface
 	
-		void SetRawFilterText(const FText& InFilterText) { TextFilter.SetRawFilterText(InFilterText); }
+		void SetRawFilterText(const FText& InFilterText) { TextFilter->SetRawFilterText(InFilterText); }
+		TSharedRef<TTextFilter<TFilterType>> GetTextFilter() const { return TextFilter; }
 
 	protected:
 		
@@ -41,34 +49,42 @@ namespace UE::MultiUserServer
 	private:
 
 		/** Does the actual string search */
-		TTextFilter<TFilterType> TextFilter;
+		TSharedRef<TTextFilter<TFilterType>> TextFilter;
+		
+		typename IFilter<TFilterType>::FChangedEvent OnChangedEvent;
 	};
 
 	/** Creates a search bar */
 	template<typename TTextSearchFilterType /* Expected subtype of TConcertFilter_TextSearch */, typename TFilterType>
-	class TConcertFrontendFilter_TextSearch : public TConcertFrontendFilterAggregate<TTextSearchFilterType, TFilterType, SSearchBox>
+	class TConcertFrontendFilter_TextSearch : public TConcertFrontendFilterAggregate<TTextSearchFilterType, TFilterType, SFilterSearchBox>
 	{
-		using Super = TConcertFrontendFilterAggregate<TTextSearchFilterType, TFilterType, SSearchBox>;
+		using Super = TConcertFrontendFilterAggregate<TTextSearchFilterType, TFilterType, SFilterSearchBox>;
 	public:
 
 		template<typename... TArg>
-		TConcertFrontendFilter_TextSearch(TArg&&... Arg)
-			: Super(Forward<TArg>(Arg)...)
+		TConcertFrontendFilter_TextSearch(TSharedPtr<FFilterCategory> InCategory, TArg&&... Arg)
+			: Super(MoveTemp(InCategory), Forward<TArg>(Arg)...)
 		{
-			Super::ChildSlot = SNew(SSearchBox)
+			SAssignNew(SearchBox, SFilterSearchBox)
+				.DelayChangeNotificationsWhileTyping(true)
 				.OnTextChanged_Lambda([this](const FText& NewSearchText)
 				{
 					SearchText = NewSearchText;
 					Super::Implementation.SetRawFilterText(NewSearchText);
-				})
-				.DelayChangeNotificationsWhileTyping(true);
+				});
 		}
 		
+		virtual FString GetName() const override { return FString("Text Search"); }
+		virtual FText GetDisplayName() const override { return NSLOCTEXT("UnrealMultiUserUI.TConcertFrontendFilter_TextSearch", "DisplayLabel", "Text Search"); }
+		
 		const FText& GetSearchText() const { return SearchText; }
+		TSharedRef<SFilterSearchBox> GetSearchBox() const { return SearchBox.ToSharedRef(); }
+		TSharedRef<TTextFilter<TFilterType>> MakeTextFilter() const { return Super::Implementation.MakeTextFilter(); }
 
 	private:
 
 		FText SearchText;
+		TSharedPtr<SFilterSearchBox> SearchBox; 
 	};
 }
 
diff --git a/Engine/Plugins/Developer/Concert/ConcertApp/MultiUserServer/Source/MultiUserServer/Private/Widgets/Util/Filter/ConcertRootFrontendFilter.h b/Engine/Plugins/Developer/Concert/ConcertApp/MultiUserServer/Source/MultiUserServer/Private/Widgets/Util/Filter/ConcertRootFrontendFilter.h
--- a/Engine/Plugins/Developer/Concert/ConcertApp/MultiUserServer/Source/MultiUserServer/Private/Widgets/Util/Filter/ConcertRootFrontendFilter.h
+++ b/Engine/Plugins/Developer/Concert/ConcertApp/MultiUserServer/Source/MultiUserServer/Private/Widgets/Util/Filter/ConcertRootFrontendFilter.h
@@ -3,9 +3,9 @@
 #pragma once
 
 #include "CoreMinimal.h"
-#include "ConcertFilter.h"
 #include "ConcertFrontendFilter.h"
-#include "ConcertFrontendFilter_TextSearch.h"
+#include "Misc/IFilter.h"
+#include "SConcertFilterBar.h"
 
 #include "Algo/AllOf.h"
 #include "Widgets/SBoxPanel.h"
@@ -16,107 +16,112 @@ namespace UE::MultiUserServer
 	template<typename TFilterType, typename TTextSearchFilterType>
 	class TConcertFrontendRootFilter 
 		:
-		public TConcertFilter<TFilterType>,
+		public IFilter<TFilterType>,
 		public TSharedFromThis<TConcertFrontendRootFilter<TFilterType, TTextSearchFilterType>>
 	{
-		using Super = TConcertFilter<TFilterType>;
 	public:
 
-		using FConcertFilter = TConcertFilter<TFilterType>;
-		using FConcertFrontendFilter = TConcertFrontendFilter<TFilterType>;
-		using FConcertFrontendFilter_TextSearch = TTextSearchFilterType;
-
 		TConcertFrontendRootFilter(
-			TSharedRef<FConcertFrontendFilter_TextSearch> TextSearchFilter,
-			TArray<TSharedRef<FConcertFrontendFilter>> FrontendFilters,
-			const TArray<TSharedRef<FConcertFilter>>& NonVisualFilters = {}
+			TSharedRef<TTextSearchFilterType> InTextSearchFilter,
+			TArray<TSharedRef<TConcertFrontendFilter<TFilterType>>> InFrontendFilters,
+			TArray<TSharedRef<IFilter<TFilterType>>> InNonVisualFilters = {},
+			TSharedRef<FFilterCategory> InDefaultCategory = MakeShared<FFilterCategory>(NSLOCTEXT("UnrealMultiUserServer.TConcertFrontendRootFilter", "DefaultCategoryLabel", "Default"), FText::GetEmpty())
 			)
-			: TextSearchFilter(MoveTemp(TextSearchFilter))
-			, FrontendFilters(MoveTemp(FrontendFilters))
-			, AllFilters(this->FrontendFilters)
+			: DefaultCategory(MoveTemp(InDefaultCategory))
+			, TextSearchFilter(MoveTemp(InTextSearchFilter))
+			, NonVisualFilters(MoveTemp(InNonVisualFilters))
 		{
-			AllFilters.Reserve(FrontendFilters.Num() + 1 + NonVisualFilters.Num());
-			AllFilters.Add(TextSearchFilter);
-			for (const TSharedRef<FConcertFilter>& NonVisualFilter : NonVisualFilters)
+			FilterBar = SNew(SConcertFilterBar<TFilterType>)
+				.FilterSearchBox(TextSearchFilter->GetSearchBox())
+				.CreateTextFilter(SBasicFilterBar<TFilterType>::FCreateTextFilter::CreateLambda([this]()
+				{
+					return MakeShared<FCustomTextFilter<TFilterType>>(TextSearchFilter->MakeTextFilter());
+				}))
+				.CustomFilters(InFrontendFilters)
+				.OnFilterChanged_Raw(this, &TConcertFrontendRootFilter<TFilterType, TTextSearchFilterType>::BroadcastOnChanged);
+
+			TextSearchFilter->OnChanged().AddRaw(this, &TConcertFrontendRootFilter<TFilterType, TTextSearchFilterType>::BroadcastOnChanged);
+			for (const TSharedRef<IFilter<TFilterType>>& Filter : InFrontendFilters)
 			{
-				AllFilters.Add(NonVisualFilter);
+				Filter->OnChanged().AddRaw(this, &TConcertFrontendRootFilter<TFilterType, TTextSearchFilterType>::BroadcastOnChanged);
 			}
-	
-			for (const TSharedRef<FConcertFilter>& Filter : AllFilters)
+			for (const TSharedRef<IFilter<TFilterType>>& Filter : NonVisualFilters)
 			{
 				Filter->OnChanged().AddRaw(this, &TConcertFrontendRootFilter<TFilterType, TTextSearchFilterType>::BroadcastOnChanged);
 			}
 		}
 
 		/** Builds the widget view for all contained filters */
-		TSharedRef<SWidget> BuildFilterWidgets() const
+		TSharedRef<SWidget> BuildFilterWidgets()
 		{
 			return SNew(SVerticalBox)
 
-				// Search bar
 				+SVerticalBox::Slot()
 				.AutoHeight()
 				[
-					TextSearchFilter->GetFilterWidget()
+					SNew(SHorizontalBox)
+
+					// Filter button
+					+SHorizontalBox::Slot()
+					.AutoWidth()
+					[
+						SBasicFilterBar<TFilterType>::MakeAddFilterButton(FilterBar.ToSharedRef())
+					]
+
+					// Search bar
+					+SHorizontalBox::Slot()
+					.FillWidth(1.f)
+					.Padding(0, 2)
+					[
+						TextSearchFilter->GetSearchBox()
+					]
 				]
 
 				+SVerticalBox::Slot()
 				.AutoHeight()
-				.Padding(0, 2)
 				[
-					BuildCustomFilterListWidget()
-				];
+					FilterBar.ToSharedRef()
+				]
+			;
 		}
 	
 		//~ Begin IFilter Interface
 		virtual bool PassesFilter(TFilterType InItem) const override
 		{
-			return Algo::AllOf(
-				AllFilters,
-				[&InItem](const TSharedRef<FConcertFilter>& AndFilter){ return AndFilter->PassesFilter(InItem); }
+			const bool bPassesActiveFilters = FilterBar->GetAllActiveFilters()->PassesAllFilters(InItem);
+			const bool bPassesTextSearch = TextSearchFilter->PassesFilter(InItem);
+			const bool bPassesNonVisualFitlers = Algo::AllOf(
+				NonVisualFilters,
+				[&InItem](const TSharedRef<IFilter<TFilterType>>& AndFilter){ return AndFilter->PassesFilter(InItem); }
 				);
+
+			return bPassesActiveFilters && bPassesTextSearch && bPassesNonVisualFitlers;
 		}
+		virtual typename IFilter<TFilterType>::FChangedEvent& OnChanged() override { return OnChangedEvent; }
 		//~ End IFilter Interface
 
-		FORCEINLINE const TSharedRef<FConcertFrontendFilter_TextSearch>& GetTextSearchFilter() const { return TextSearchFilter; }
-
+		FORCEINLINE const TSharedRef<TTextSearchFilterType>& GetTextSearchFilter() const { return TextSearchFilter; }
+		FORCEINLINE TSharedRef<FFilterCategory> GetDefaultCategory() const { return DefaultCategory; }
+		
 	private:
 
+		/** The default category that filters are added to */
+		TSharedRef<FFilterCategory> DefaultCategory;
+		/** Displays the currently active filters */
+		TSharedPtr<SConcertFilterBar<TFilterType>> FilterBar;
+		
 		/** The text search filter. Also in AllFilters. Separate variable to build search bar in new line. */
-		TSharedRef<FConcertFrontendFilter_TextSearch> TextSearchFilter;
-
-		/** AllFrontendFilters without TextSearchFilter. */
-		TArray<TSharedRef<FConcertFrontendFilter>> FrontendFilters;
+		TSharedRef<TTextSearchFilterType> TextSearchFilter;
 	
 		/** Filters that are combined using logical AND. */
-		TArray<TSharedRef<FConcertFilter>> AllFilters;
+		TArray<TSharedRef<IFilter<TFilterType>>> NonVisualFilters;
 
 		void BroadcastOnChanged()
 		{
-			Super::OnChanged().Broadcast();
-		}
-
-		/** Builds the widgets that go under the text */
-		TSharedRef<SWidget> BuildCustomFilterListWidget() const
-		{
-			TSharedRef<SHorizontalBox> Box = SNew(SHorizontalBox);
-
-			bool bIsFirst = true;
-			for (const TSharedRef<FConcertFrontendFilter>& Filter : FrontendFilters)
-			{
-				const FMargin Margin = bIsFirst ? FMargin() : FMargin(8, 0, 0, 0);
-				bIsFirst = false;
-				Box->AddSlot()
-					.AutoWidth()
-					.Padding(Margin)
-					.VAlign(VAlign_Center)
-					[
-						Filter->GetFilterWidget()
-					];
-			}
-
-			return Box;
+			OnChangedEvent.Broadcast();
 		}
+		
+		typename IFilter<TFilterType>::FChangedEvent OnChangedEvent;
 	};
 }
 
diff --git a/Engine/Plugins/Developer/Concert/ConcertApp/MultiUserServer/Source/MultiUserServer/Private/Widgets/Util/Filter/SConcertFilterBar.h b/Engine/Plugins/Developer/Concert/ConcertApp/MultiUserServer/Source/MultiUserServer/Private/Widgets/Util/Filter/SConcertFilterBar.h
new file mode 100644
--- /dev/null
+++ b/Engine/Plugins/Developer/Concert/ConcertApp/MultiUserServer/Source/MultiUserServer/Private/Widgets/Util/Filter/SConcertFilterBar.h
@@ -0,0 +1,146 @@
+// Copyright Epic Games, Inc. All Rights Reserved.
+
+#pragma once
+
+#include "CoreMinimal.h"
+#include "Algo/Transform.h"
+#include "ConcertLogGlobal.h"
+#include "ConcertFrontendFilter.h"
+#include "Filters/SBasicFilterBar.h"
+#include "Widgets/DeclarativeSyntaxSupport.h"
+#include "SConcertFilterBar.generated.h"
+
+UCLASS()
+class UConcertFilterBarContext : public UObject
+{
+	GENERATED_BODY()
+public:
+
+	DECLARE_DELEGATE_OneParam(FOnPopulateMenu, UToolMenu*);
+	FOnPopulateMenu PopulateMenu;
+};
+
+namespace UE::MultiUserServer
+{
+	template<typename TFilterType>
+	class SConcertFilterBar : public SBasicFilterBar<TFilterType>
+	{
+		using Super = SBasicFilterBar<TFilterType>;
+	public:
+
+		using FOnFilterChanged = typename SBasicFilterBar<TFilterType>::FOnFilterChanged;
+		using FCreateTextFilter = typename SBasicFilterBar<TFilterType>::FCreateTextFilter;
+
+		SLATE_BEGIN_ARGS(SConcertFilterBar)
+		{}
+
+			/** Delegate for when filters have changed */
+			SLATE_EVENT(SConcertFilterBar<TFilterType>::FOnFilterChanged, OnFilterChanged)
+		
+			/** Initial List of Custom Filters that will be added to the AddFilter Menu */
+			SLATE_ARGUMENT(TArray<TSharedRef<TConcertFrontendFilter<TFilterType>>>, CustomFilters)
+		
+			/** A delegate to create a Text Filter for FilterType items. If provided, will allow creation of custom
+			 * text filters from the filter dropdown menu.
+			 * @see FCustomTextFilter
+			 */
+			SLATE_ARGUMENT(FCreateTextFilter, CreateTextFilter)
+		
+			/** An SFilterSearchBox that can be attached to this filter bar. When provided along with a CreateTextFilter
+			 *  delegate, allows the user to save searches from the Search Box as text filters for the filter bar.
+			 *	NOTE: Will bind a delegate to SFilterSearchBox::OnSaveSearchClicked
+			 */
+			SLATE_ARGUMENT(TSharedPtr<SFilterSearchBox>, FilterSearchBox)
+		
+		SLATE_END_ARGS()
+
+		void Construct(const FArguments& InArgs)
+		{
+			TArray<TSharedRef<FFilterBase<TFilterType>>> TransformedCustomFilters;
+			Algo::Transform(InArgs._CustomFilters, TransformedCustomFilters, [](const TSharedRef<TConcertFrontendFilter<TFilterType>>& Item){ return Item; });
+			
+			typename SBasicFilterBar<TFilterType>::FArguments Args;
+			Args._OnFilterChanged = InArgs._OnFilterChanged;
+			Args._CustomFilters = MoveTemp(TransformedCustomFilters);
+			Args._CreateTextFilter = InArgs._CreateTextFilter;
+			Args._FilterSearchBox = InArgs._FilterSearchBox;
+		
+			SBasicFilterBar<TFilterType>::Construct(Args);
+		}
+
+	private:
+		
+		virtual TSharedRef<SWidget> MakeAddFilterMenu() override
+		{
+			const FName FilterMenuName = "ConcertFilterBar.ConcertFilterBar";
+			if (!UToolMenus::Get()->IsMenuRegistered(FilterMenuName))
+			{
+				UToolMenu* Menu = UToolMenus::Get()->RegisterMenu(FilterMenuName);
+				Menu->bShouldCloseWindowAfterMenuSelection = true;
+				Menu->bCloseSelfOnly = true;
+
+				Menu->AddDynamicSection(NAME_None, FNewToolMenuDelegate::CreateLambda([this](UToolMenu* InMenu)
+				{
+					if (const UConcertFilterBarContext* Context = InMenu->FindContext<UConcertFilterBarContext>())
+					{
+						Context->PopulateMenu.Execute(InMenu);
+					}
+				}));
+			}
+
+			UConcertFilterBarContext* Context = NewObject<UConcertFilterBarContext>();
+			Context->PopulateMenu = UConcertFilterBarContext::FOnPopulateMenu::CreateSP(this, &SConcertFilterBar::HandlePopulateAddFilterMenu);
+			
+			const FToolMenuContext ToolMenuContext(Context);
+			return UToolMenus::Get()->GenerateWidget(FilterMenuName, ToolMenuContext);
+		}
+
+		void HandlePopulateAddFilterMenu(UToolMenu* InMenu)
+		{
+			Super::PopulateCommonFilterSections(InMenu);
+
+			UE_CLOG(this->AllFilterCategories.Num() > 1, LogConcert, Warning, TEXT("This widget is designed for uncategorisied filters"));
+			FToolMenuSection& Section = InMenu->AddSection("BasicFilterBarFiltersMenu");
+			for (const TSharedRef<FFilterBase<TFilterType>>& Filter : this->AllFrontendFilters)
+			{
+				const TSharedRef<TConcertFrontendFilter<TFilterType>> CastFilter = StaticCastSharedRef<TConcertFrontendFilter<TFilterType>>(Filter);
+				Section.AddSubMenu(
+					NAME_None,
+					Filter->GetDisplayName(),
+					Filter->GetToolTipText(),
+					FNewMenuDelegate::CreateLambda([WeakFilter = TWeakPtr<TConcertFrontendFilter<TFilterType>>(CastFilter)](FMenuBuilder& MenuBuilder)
+					{
+						const TSharedPtr<TConcertFrontendFilter<TFilterType>> PinnedFilter = WeakFilter.Pin();
+						if (ensure(PinnedFilter))
+						{
+							PinnedFilter->ExposeEditWidgets(MenuBuilder);
+						}
+					}),
+					FUIAction(
+						FExecuteAction::CreateLambda([this, WeakFilter = TWeakPtr<TConcertFrontendFilter<TFilterType>>(CastFilter)]()
+						{
+							const TSharedPtr<TConcertFrontendFilter<TFilterType>> PinnedFilter = WeakFilter.Pin();
+							if (ensure(PinnedFilter))
+							{
+								Super::FrontendFilterClicked(PinnedFilter.ToSharedRef());
+							}
+						}),
+						FCanExecuteAction(),
+						FGetActionCheckState::CreateLambda([this, WeakFilter = TWeakPtr<TConcertFrontendFilter<TFilterType>>(CastFilter)]()
+						{
+							const TSharedPtr<TConcertFrontendFilter<TFilterType>> PinnedFilter = WeakFilter.Pin();
+							if (ensure(PinnedFilter))
+							{
+								return Super::IsFilterActive(PinnedFilter)
+									? ECheckBoxState::Checked
+									: ECheckBoxState::Unchecked;
+							}
+							return ECheckBoxState::Undetermined;
+						})),
+					EUserInterfaceActionType::ToggleButton
+				);
+			}
+		}
+	};
+}
+
