Add options button for client browser

diff --git a/Engine/Plugins/Developer/Concert/ConcertApp/MultiUserServer/Source/MultiUserServer/Private/Widgets/Clients/Browser/SConcertNetworkBrowser.cpp b/Engine/Plugins/Developer/Concert/ConcertApp/MultiUserServer/Source/MultiUserServer/Private/Widgets/Clients/Browser/SConcertNetworkBrowser.cpp
--- a/Engine/Plugins/Developer/Concert/ConcertApp/MultiUserServer/Source/MultiUserServer/Private/Widgets/Clients/Browser/SConcertNetworkBrowser.cpp
+++ b/Engine/Plugins/Developer/Concert/ConcertApp/MultiUserServer/Source/MultiUserServer/Private/Widgets/Clients/Browser/SConcertNetworkBrowser.cpp
@@ -42,6 +42,7 @@ void UE::MultiUserServer::SConcertNetworkBrowser::Construct(const FArguments& In
 		SNew(SVerticalBox)
 
 		+SVerticalBox::Slot()
+		.Padding(4.f)
 		.AutoHeight()
 		[
 			CreateSearchArea(InArgs)
@@ -93,7 +94,7 @@ TSharedRef<SWidget> UE::MultiUserServer::SConcertNetworkBrowser::CreateSearchAre
 		.AutoWidth()
 		[
 			SNew(SComboButton)
-			.OnGetMenuContent(this, &SConcertNetworkBrowser::MakeSessionOption)
+			.OnGetMenuContent(this, &SConcertNetworkBrowser::CreateSessionFilterOptionMenu)
 			.ButtonContent()
 			[
 				SNew(STextBlock)
@@ -130,64 +131,22 @@ TSharedRef<SWidget> UE::MultiUserServer::SConcertNetworkBrowser::CreateSearchAre
 		.AutoWidth()
 		.HAlign(HAlign_Right)
 		[
-			CreateKeepDisconnectedClients()
+			CreateOptionsButton()
 		];
 }
 
-TSharedRef<SWidget> UE::MultiUserServer::SConcertNetworkBrowser::CreateKeepDisconnectedClients()
+TSharedRef<SWidget> UE::MultiUserServer::SConcertNetworkBrowser::CreateOptionsButton()
 {
-	return SNew(SHorizontalBox)
-		.ToolTipText(LOCTEXT("KeepDisconnectedClients.Tooltip", "Whether to keep clients that have disconnected in memory. This may be useful in unstable networks when you want to analyse why clients keep disconnecting."))
-		
-		+SHorizontalBox::Slot()
-		.AutoWidth()
-		.VAlign(VAlign_Center)
-		[
-			SNew(STextBlock)
-			.Text(LOCTEXT("KeepDisconnectedClients.Label", "Keep Disconnected"))
-			]
-
-		+SHorizontalBox::Slot()
-		.AutoWidth()
-		.Padding(4.f, 0.f, 0.f, 0.f)
-		[
-			SNew(SCheckBox)
-			.IsChecked_Lambda([this](){ return BrowserModel->ShouldKeepClientsAfterDisconnect() ? ECheckBoxState::Checked : ECheckBoxState::Unchecked; })
-			.OnCheckStateChanged_Lambda([this](ECheckBoxState NewState)
-			{
-				const bool bKeepDisconnected = NewState == ECheckBoxState::Checked;
-				if (bKeepDisconnected)
-				{
-					BrowserModel->SetKeepClientsAfterDisconnect(true);
-					return;
-				}
-				
-				const int32 NumDisconnected = Algo::CountIf(BrowserModel->GetItems(), [](const TSharedPtr<FClientBrowserItem>& Item) { return Item->IsDisconnected(); });
-				if (NumDisconnected == 0)
-				{
-					BrowserModel->SetKeepClientsAfterDisconnect(false);
-					return;
-				}
-				
-				const TSharedRef<SMessageDialog> Dialog = SNew(SMessageDialog)
-					.Title(LOCTEXT("RemoveDisconnectedClients.Title", "Remove disconnected clients?"))
-					.Icon(FAppStyle::Get().GetBrush("Icons.WarningWithColor.Large"))
-					.Message(FText::Format(LOCTEXT("RemoveDisconnectedClients.MessageFmt", "There are {0} disconnected clients. If you proceed, these clients will be removed from the session browser; opened log tabs will remain open.\nProceed?"), NumDisconnected))
-					.UseScrollBox(false)
-					.Buttons({
-						SMessageDialog::FButton(LOCTEXT("RemoveButton", "Remove"))
-							.SetOnClicked(FSimpleDelegate::CreateLambda([this]()
-							{
-								BrowserModel->SetKeepClientsAfterDisconnect(false);
-							})),
-						SMessageDialog::FButton(LOCTEXT("CancelButton", "Keep"))
-							.SetPrimary(true)
-							.SetFocus()
-					});
-				FConcertServerUIModule::Get().GetModalWindowManager()->ShowFakeModalWindow(Dialog);
-			})
-		]
-	;
+	return SNew(SComboButton)
+			.ComboButtonStyle( FAppStyle::Get(), "SimpleComboButtonWithIcon" )
+			.OnGetMenuContent( this, &SConcertNetworkBrowser::CreateOptionsButtonMenu)
+			.HasDownArrow(false)
+			.ButtonContent()
+			[
+				SNew(SImage)
+				.ColorAndOpacity(FSlateColor::UseForeground())
+				.Image( FAppStyle::Get().GetBrush("Icons.Settings") )
+			];
 }
 
 TSharedRef<SWidget> UE::MultiUserServer::SConcertNetworkBrowser::CreateTileView()
@@ -205,6 +164,60 @@ TSharedRef<SWidget> UE::MultiUserServer::SConcertNetworkBrowser::CreateTileView(
 	;
 }
 
+TSharedRef<SWidget> UE::MultiUserServer::SConcertNetworkBrowser::CreateOptionsButtonMenu()
+{
+	FMenuBuilder MenuBuilder(true, nullptr);
+
+	MenuBuilder.AddMenuEntry(
+		LOCTEXT("KeepDisconnectedClients.Label", "Keep Disconnected"),
+		LOCTEXT("KeepDisconnectedClients.Tooltip", "Whether to keep clients that have disconnected in memory. This may be useful in unstable networks when you want to analyse why clients keep disconnecting."),
+		FSlateIcon(),
+		FUIAction(
+			FExecuteAction::CreateLambda([this](){ ToggleShouldKeepDisconnected(); }),
+			FCanExecuteAction::CreateLambda([](){ return true; }),
+			FIsActionChecked::CreateLambda([this](){ return BrowserModel->ShouldKeepClientsAfterDisconnect(); })
+		),
+		NAME_None,
+		EUserInterfaceActionType::ToggleButton
+	);
+	
+	return MenuBuilder.MakeWidget();
+}
+
+void UE::MultiUserServer::SConcertNetworkBrowser::ToggleShouldKeepDisconnected()
+{
+	const bool bKeepDisconnected = !BrowserModel->ShouldKeepClientsAfterDisconnect();
+	if (bKeepDisconnected)
+	{
+		BrowserModel->SetKeepClientsAfterDisconnect(true);
+		return;
+	}
+				
+	const int32 NumDisconnected = Algo::CountIf(BrowserModel->GetItems(), [](const TSharedPtr<FClientBrowserItem>& Item) { return Item->IsDisconnected(); });
+	if (NumDisconnected == 0)
+	{
+		BrowserModel->SetKeepClientsAfterDisconnect(false);
+		return;
+	}
+				
+	const TSharedRef<SMessageDialog> Dialog = SNew(SMessageDialog)
+		.Title(LOCTEXT("RemoveDisconnectedClients.Title", "Remove disconnected clients?"))
+		.Icon(FAppStyle::Get().GetBrush("Icons.WarningWithColor.Large"))
+		.Message(FText::Format(LOCTEXT("RemoveDisconnectedClients.MessageFmt", "There are {0} disconnected clients. If you proceed, these clients will be removed from the session browser; opened log tabs will remain open.\nProceed?"), NumDisconnected))
+		.UseScrollBox(false)
+		.Buttons({
+			SMessageDialog::FButton(LOCTEXT("RemoveButton", "Remove"))
+				.SetOnClicked(FSimpleDelegate::CreateLambda([this]()
+				{
+					BrowserModel->SetKeepClientsAfterDisconnect(false);
+				})),
+			SMessageDialog::FButton(LOCTEXT("CancelButton", "Keep"))
+				.SetPrimary(true)
+				.SetFocus()
+		});
+	FConcertServerUIModule::Get().GetModalWindowManager()->ShowFakeModalWindow(Dialog);
+}
+
 void UE::MultiUserServer::SConcertNetworkBrowser::OnSessionCreated(const FGuid& SessionId)
 {
 	if (bShowAllSessions)
@@ -245,7 +258,7 @@ void UE::MultiUserServer::SConcertNetworkBrowser::OnClientListChanged(TSharedPtr
 	TileView->RequestListRefresh();
 }
 
-TSharedRef<SWidget> UE::MultiUserServer::SConcertNetworkBrowser::SConcertNetworkBrowser::MakeSessionOption()
+TSharedRef<SWidget> UE::MultiUserServer::SConcertNetworkBrowser::SConcertNetworkBrowser::CreateSessionFilterOptionMenu()
 {
 	FMenuBuilder MenuBuilder(true, nullptr);
 
diff --git a/Engine/Plugins/Developer/Concert/ConcertApp/MultiUserServer/Source/MultiUserServer/Private/Widgets/Clients/Browser/SConcertNetworkBrowser.h b/Engine/Plugins/Developer/Concert/ConcertApp/MultiUserServer/Source/MultiUserServer/Private/Widgets/Clients/Browser/SConcertNetworkBrowser.h
--- a/Engine/Plugins/Developer/Concert/ConcertApp/MultiUserServer/Source/MultiUserServer/Private/Widgets/Clients/Browser/SConcertNetworkBrowser.h
+++ b/Engine/Plugins/Developer/Concert/ConcertApp/MultiUserServer/Source/MultiUserServer/Private/Widgets/Clients/Browser/SConcertNetworkBrowser.h
@@ -80,16 +80,19 @@ namespace UE::MultiUserServer
 		FOnClientDoubleClicked OnClientDoubleClicked;
 
 		TSharedRef<SWidget> CreateSearchArea(const FArguments& InArgs);
-		TSharedRef<SWidget> CreateKeepDisconnectedClients();
+		TSharedRef<SWidget> CreateOptionsButton();
 		TSharedRef<SWidget> CreateTileView();
 
+		TSharedRef<SWidget> CreateOptionsButtonMenu();
+		void ToggleShouldKeepDisconnected();
+
 		// Model events
 		void OnSessionCreated(const FGuid& SessionId);
 		void OnSessionDestroyed(const FGuid& SessionId);
 		void OnClientListChanged(TSharedPtr<FClientBrowserItem> Item, IClientBrowserModel::EClientUpdateType UpdateType);
 
 		// Combo button
-		TSharedRef<SWidget> MakeSessionOption();
+		TSharedRef<SWidget> CreateSessionFilterOptionMenu();
 
 		FText GetErrorMessageText() const;
 		
diff --git a/Engine/Plugins/Developer/Concert/ConcertApp/MultiUserServer/Source/MultiUserServer/Private/Widgets/Clients/SConcertClientsTabView.cpp b/Engine/Plugins/Developer/Concert/ConcertApp/MultiUserServer/Source/MultiUserServer/Private/Widgets/Clients/SConcertClientsTabView.cpp
--- a/Engine/Plugins/Developer/Concert/ConcertApp/MultiUserServer/Source/MultiUserServer/Private/Widgets/Clients/SConcertClientsTabView.cpp
+++ b/Engine/Plugins/Developer/Concert/ConcertApp/MultiUserServer/Source/MultiUserServer/Private/Widgets/Clients/SConcertClientsTabView.cpp
@@ -226,7 +226,7 @@ namespace UE::MultiUserServer
 				[this]()
 				{
 					const bool bIsGlobalLogOpen = IsGlobalLogOpen();
-					return bIsGlobalLogOpen ? EVisibility::Hidden : EVisibility::Visible;
+					return bIsGlobalLogOpen ? EVisibility::Collapsed : EVisibility::Visible;
 				})
 			.OnClicked_Lambda([this]()
 			{
