Replace view options in package and log viewer

diff --git a/Engine/Plugins/Developer/Concert/ConcertApp/MultiUserServer/Source/MultiUserServer/Private/Widgets/Clients/Logging/SConcertTransportLog.cpp b/Engine/Plugins/Developer/Concert/ConcertApp/MultiUserServer/Source/MultiUserServer/Private/Widgets/Clients/Logging/SConcertTransportLog.cpp
--- a/Engine/Plugins/Developer/Concert/ConcertApp/MultiUserServer/Source/MultiUserServer/Private/Widgets/Clients/Logging/SConcertTransportLog.cpp
+++ b/Engine/Plugins/Developer/Concert/ConcertApp/MultiUserServer/Source/MultiUserServer/Private/Widgets/Clients/Logging/SConcertTransportLog.cpp
@@ -57,7 +57,16 @@ namespace UE::MultiUserServer
 					.AutoHeight()
 					.VAlign(VAlign_Top)
 					[
-						InArgs._Filter ? InArgs._Filter->BuildFilterWidgets() : SNullWidget::NullWidget
+						InArgs._Filter
+							? InArgs._Filter->BuildFilterWidgets
+							(
+								FFilterWidgetArgs()
+								.RightOfSearchBar
+								(
+									CreateViewOptionsButton()
+								)
+							)
+							: SNullWidget::NullWidget
 					]
 
 					+SVerticalBox::Slot()
@@ -71,7 +80,6 @@ namespace UE::MultiUserServer
 					.AutoHeight()
 					[
 						SNew(SConcertTransportLogFooter, PagedLogList.ToSharedRef())
-						.ExtendViewOptions(this, &SConcertTransportLog::ExtendViewOptions)
 					]
 				]
 			]
@@ -105,13 +113,29 @@ namespace UE::MultiUserServer
 		ScrollToLog(Index);
 	}
 
+	TSharedRef<SWidget> SConcertTransportLog::CreateViewOptionsButton()
+	{
+		return SNew(SComboButton)
+			.ComboButtonStyle(FAppStyle::Get(), "SimpleComboButtonWithIcon")
+			.OnGetMenuContent(this, &SConcertTransportLog::CreateViewOptionsMenu)
+			.HasDownArrow(false)
+			.ButtonContent()
+			[
+				SNew(SImage)
+				.ColorAndOpacity(FSlateColor::UseForeground())
+				.Image(FAppStyle::Get().GetBrush("Icons.Settings"))
+			];
+	}
+
 	TSharedRef<SWidget> SConcertTransportLog::CreateTableView()
 	{
+		CreateHeaderRow();
 		return SAssignNew(LogView, SListView<TSharedPtr<FConcertLogEntry>>)
 			.ListItemsSource(&PagedLogList->GetPageView())
 			.OnGenerateRow(this, &SConcertTransportLog::OnGenerateActivityRowWidget)
+			.OnContextMenuOpening_Lambda([this](){ return ConcertSharedSlate::MakeTableContextMenu(HeaderRow.ToSharedRef(), GetDefaultColumnVisibilities(), true); })
 			.SelectionMode(ESelectionMode::Multi)
-			.HeaderRow(CreateHeaderRow());
+			.HeaderRow(HeaderRow);
 	}
 
 	TSharedRef<SHeaderRow> SConcertTransportLog::CreateHeaderRow()
@@ -165,12 +189,7 @@ namespace UE::MultiUserServer
 			SHeaderRow::FColumn::FArguments Args = SHeaderRow::FColumn::FArguments()
 				.ColumnId(ColumnId)
 				.DefaultLabel(FText::FromString(PropertyName))
-				.HAlignCell(HAlign_Center)
-				// Add option to hide
-				.OnGetMenuContent_Lambda([this, ColumnId]()
-				{
-					return UE::ConcertSharedSlate::MakeHideColumnContextMenu(HeaderRow.ToSharedRef(), ColumnId);
-				});
+				.HAlignCell(HAlign_Center);
 			
 			if (const float* ManualSize = ManuallySizeColumns.Find(ColumnId))
 			{
@@ -220,22 +239,30 @@ namespace UE::MultiUserServer
 
 	void SConcertTransportLog::RestoreDefaultColumnVisiblities()
 	{
-		ForEachPropertyColumn([this](const FProperty& Property, FName ColumnId)
+		const TMap<FName, bool> DefaultVisibilities = GetDefaultColumnVisibilities();
+		ForEachPropertyColumn([this, &DefaultVisibilities](const FProperty& Property, FName ColumnId)
 		{
-			const TSet<FName> HiddenByDefault = {
-				GET_MEMBER_NAME_CHECKED(FConcertLog, Frame),
-				GET_MEMBER_NAME_CHECKED(FConcertLog, MessageId),
-				GET_MEMBER_NAME_CHECKED(FConcertLog, MessageOrderIndex),
-				GET_MEMBER_NAME_CHECKED(FConcertLog, ChannelId),
-				GET_MEMBER_NAME_CHECKED(FConcertLog, CustomPayloadTypename),
-				GET_MEMBER_NAME_CHECKED(FConcertLog, StringPayload)
-			};
-			HeaderRow->SetShowGeneratedColumn(ColumnId, !HiddenByDefault.Contains(ColumnId));
+			const bool* Visibility = DefaultVisibilities.Find(ColumnId);
+			HeaderRow->SetShowGeneratedColumn(ColumnId, Visibility ? *Visibility : true);
 		});
 	}
 
-	void SConcertTransportLog::ExtendViewOptions(FMenuBuilder& MenuBuilder)
+	TMap<FName, bool> SConcertTransportLog::GetDefaultColumnVisibilities() const
 	{
+		return {
+			{ GET_MEMBER_NAME_CHECKED(FConcertLog, Frame), false },
+			{ GET_MEMBER_NAME_CHECKED(FConcertLog, MessageId), false },
+			{ GET_MEMBER_NAME_CHECKED(FConcertLog, MessageOrderIndex), false },
+			{ GET_MEMBER_NAME_CHECKED(FConcertLog, ChannelId), false },
+			{ GET_MEMBER_NAME_CHECKED(FConcertLog, CustomPayloadTypename), false },
+			{ GET_MEMBER_NAME_CHECKED(FConcertLog, StringPayload), false}
+		};
+	}
+
+	TSharedRef<SWidget> SConcertTransportLog::CreateViewOptionsMenu()
+	{
+		FMenuBuilder MenuBuilder(true, nullptr);
+		
 		MenuBuilder.AddMenuEntry(
 			LOCTEXT("AutoScroll", "Auto Scroll"),
 			LOCTEXT("AutoScroll_Tooltip", "Automatically scroll as new logs arrive (affects last page)"),
@@ -266,51 +293,7 @@ namespace UE::MultiUserServer
 				EUserInterfaceActionType::ToggleButton
 			);
 
-		MenuBuilder.AddSeparator();
-		MenuBuilder.AddMenuEntry(
-				LOCTEXT("SelectAll", "Show all"),
-				FText::GetEmpty(),
-				FSlateIcon(),
-				FUIAction(
-					FExecuteAction::CreateLambda([this]()
-					{
-						ForEachPropertyColumn([this](const FProperty& Property, FName ColumnId)
-						{
-							HeaderRow->SetShowGeneratedColumn(ColumnId, true);
-						});
-					}),
-					FCanExecuteAction::CreateLambda([] { return true; })),
-				NAME_None,
-				EUserInterfaceActionType::Button
-			);
-		MenuBuilder.AddMenuEntry(
-				LOCTEXT("HideAll", "Hide all"),
-				FText::GetEmpty(),
-				FSlateIcon(),
-				FUIAction(
-					FExecuteAction::CreateLambda([this]()
-					{
-						ForEachPropertyColumn([this](const FProperty& Property, FName ColumnId)
-						{
-							HeaderRow->SetShowGeneratedColumn(ColumnId, false);
-						});
-					}),
-					FCanExecuteAction::CreateLambda([] { return true; })),
-				NAME_None,
-				EUserInterfaceActionType::Button
-			);
-		MenuBuilder.AddMenuEntry(
-				LOCTEXT("RestoreDefaultColumnVisibility", "Restore columns visibility"),
-				FText::GetEmpty(),
-				FSlateIcon(),
-				FUIAction(
-					FExecuteAction::CreateSP(this, &SConcertTransportLog::RestoreDefaultColumnVisiblities),
-					FCanExecuteAction::CreateLambda([] { return true; })),
-				NAME_None,
-				EUserInterfaceActionType::Button
-			);
-		MenuBuilder.AddSeparator();
-		UE::ConcertSharedSlate::AddEntriesForShowingHiddenRows(HeaderRow.ToSharedRef(), MenuBuilder);
+		return MenuBuilder.MakeWidget();
 	}
 
 	void SConcertTransportLog::OnFilterMenuChecked()
@@ -345,7 +328,7 @@ namespace UE::MultiUserServer
 	void SConcertTransportLog::OnColumnVisibilitySettingsChanged(const FColumnVisibilitySnapshot& ColumnSnapshot)
 	{
 		TGuardValue<bool> GuardValue(bIsUpdatingColumnVisibility, true);
-		UE::ConcertSharedSlate::RestoreColumnVisibilityState(HeaderRow.ToSharedRef(), ColumnSnapshot);
+		ConcertSharedSlate::RestoreColumnVisibilityState(HeaderRow.ToSharedRef(), ColumnSnapshot);
 	}
 
 	void SConcertTransportLog::OnConcertLoggingEnabledChanged(bool bNewEnabled)
diff --git a/Engine/Plugins/Developer/Concert/ConcertApp/MultiUserServer/Source/MultiUserServer/Private/Widgets/Clients/Logging/SConcertTransportLog.h b/Engine/Plugins/Developer/Concert/ConcertApp/MultiUserServer/Source/MultiUserServer/Private/Widgets/Clients/Logging/SConcertTransportLog.h
--- a/Engine/Plugins/Developer/Concert/ConcertApp/MultiUserServer/Source/MultiUserServer/Private/Widgets/Clients/Logging/SConcertTransportLog.h
+++ b/Engine/Plugins/Developer/Concert/ConcertApp/MultiUserServer/Source/MultiUserServer/Private/Widgets/Clients/Logging/SConcertTransportLog.h
@@ -74,6 +74,7 @@ namespace UE::MultiUserServer
 		bool bIsUpdatingColumnVisibility = false;
 
 		// Table view creation
+		TSharedRef<SWidget> CreateViewOptionsButton();
 		TSharedRef<SWidget> CreateTableView();
 		TSharedRef<SHeaderRow> CreateHeaderRow();
 		TSharedRef<ITableRow> OnGenerateActivityRowWidget(TSharedPtr<FConcertLogEntry> Item, const TSharedRef<STableViewBase>& OwnerTable) const;
@@ -81,7 +82,9 @@ namespace UE::MultiUserServer
 		static void ForEachPropertyColumn(TFunctionRef<void(const FProperty& ColumnPropertyId, FName ColumnId)> Callback);
 
 		void RestoreDefaultColumnVisiblities();
-		void ExtendViewOptions(FMenuBuilder& MenuBuilder);
+		TMap<FName, bool> GetDefaultColumnVisibilities() const;
+		
+		TSharedRef<SWidget> CreateViewOptionsMenu();
 		void OnFilterMenuChecked();
 		void OnPageViewChanged(const TArray<TSharedPtr<FConcertLogEntry>>&);
 		void OnColumnVisibilitySettingsChanged(const FColumnVisibilitySnapshot& ColumnSnapshot);
diff --git a/Engine/Plugins/Developer/Concert/ConcertApp/MultiUserServer/Source/MultiUserServer/Private/Widgets/Clients/Logging/SConcertTransportLogFooter.cpp b/Engine/Plugins/Developer/Concert/ConcertApp/MultiUserServer/Source/MultiUserServer/Private/Widgets/Clients/Logging/SConcertTransportLogFooter.cpp
--- a/Engine/Plugins/Developer/Concert/ConcertApp/MultiUserServer/Source/MultiUserServer/Private/Widgets/Clients/Logging/SConcertTransportLogFooter.cpp
+++ b/Engine/Plugins/Developer/Concert/ConcertApp/MultiUserServer/Source/MultiUserServer/Private/Widgets/Clients/Logging/SConcertTransportLogFooter.cpp
@@ -108,44 +108,8 @@ namespace UE::MultiUserServer
 			[
 				SNew(SSpacer)
 			]
-
-			+SHorizontalBox::Slot()
-			.AutoWidth()
-			.VAlign(VAlign_Center)
-			[
-				SNew(SComboButton)
-					.OnGetMenuContent(this, &SConcertTransportLogFooter::MakeViewOptionsMenuWidget, InArgs._ExtendViewOptions)
-					.HasDownArrow(true)
-					.ButtonContent()
-					[
-						SNew(SHorizontalBox)
-
-						+SHorizontalBox::Slot()
-						.AutoWidth()
-						.VAlign(VAlign_Center)
-						[
-							SNew(SImage).Image(FAppStyle::Get().GetBrush("Icons.Visible")) // The eye ball image.
-						]
-
-						+SHorizontalBox::Slot()
-						.AutoWidth()
-						.Padding(2, 0, 0, 0)
-						.VAlign(VAlign_Center)
-						[
-							SNew(STextBlock).Text(LOCTEXT("ViewOptions", "View Options"))
-						]
-					]
-			]
 		];
 	}
-
-	TSharedRef<SWidget> SConcertTransportLogFooter::MakeViewOptionsMenuWidget(FExtendContextMenu ExtendViewOptions)
-	{
-		FMenuBuilder MenuBuilder(/*bInShouldCloseWindowAfterMenuSelection=*/true, nullptr);
-
-		ExtendViewOptions.Execute(MenuBuilder);
-		return MenuBuilder.MakeWidget();
-	}
 }
 
 #undef LOCTEXT_NAMESPACE
diff --git a/Engine/Plugins/Developer/Concert/ConcertApp/MultiUserServer/Source/MultiUserServer/Private/Widgets/Clients/Logging/SConcertTransportLogFooter.h b/Engine/Plugins/Developer/Concert/ConcertApp/MultiUserServer/Source/MultiUserServer/Private/Widgets/Clients/Logging/SConcertTransportLogFooter.h
--- a/Engine/Plugins/Developer/Concert/ConcertApp/MultiUserServer/Source/MultiUserServer/Private/Widgets/Clients/Logging/SConcertTransportLogFooter.h
+++ b/Engine/Plugins/Developer/Concert/ConcertApp/MultiUserServer/Source/MultiUserServer/Private/Widgets/Clients/Logging/SConcertTransportLogFooter.h
@@ -20,12 +20,10 @@ namespace UE::MultiUserServer
 	public:
 	
 		DECLARE_DELEGATE_OneParam(FExtendContextMenu, FMenuBuilder&)
+			SLATE_BEGIN_ARGS(SConcertTransportLogFooter) {}
+		SLATE_END_ARGS()
 
-		SLATE_BEGIN_ARGS(SConcertTransportLogFooter) {}
-		SLATE_EVENT(FExtendContextMenu, ExtendViewOptions)
-	SLATE_END_ARGS()
-
-	void Construct(const FArguments& InArgs, TSharedRef<FPagedFilteredConcertLogList> InPagedLogList);
+		void Construct(const FArguments& InArgs, TSharedRef<FPagedFilteredConcertLogList> InPagedLogList);
 
 	private:
 
@@ -36,7 +34,5 @@ namespace UE::MultiUserServer
 		TSharedPtr<SSpinBox<FPagedFilteredConcertLogList::FPageCount>> CurrentPage;
 		/** Displays the number of pages */
 		TSharedPtr<STextBlock> PageCounterText;
-
-		TSharedRef<SWidget> MakeViewOptionsMenuWidget(FExtendContextMenu ExtendViewOptions);
 	};
 }
diff --git a/Engine/Plugins/Developer/Concert/ConcertApp/MultiUserServer/Source/MultiUserServer/Private/Widgets/Clients/PackageTransmission/SPackageTransmissionTable.cpp b/Engine/Plugins/Developer/Concert/ConcertApp/MultiUserServer/Source/MultiUserServer/Private/Widgets/Clients/PackageTransmission/SPackageTransmissionTable.cpp
--- a/Engine/Plugins/Developer/Concert/ConcertApp/MultiUserServer/Source/MultiUserServer/Private/Widgets/Clients/PackageTransmission/SPackageTransmissionTable.cpp
+++ b/Engine/Plugins/Developer/Concert/ConcertApp/MultiUserServer/Source/MultiUserServer/Private/Widgets/Clients/PackageTransmission/SPackageTransmissionTable.cpp
@@ -10,6 +10,8 @@
 #include "SPackageTransmissionTableRow.h"
 
 #include "Framework/MultiBox/MultiBoxBuilder.h"
+#include "Settings/MultiUserServerPackageTransmissionSettings.h"
+#include "Widgets/Input/SComboButton.h"
 #include "Widgets/SBoxPanel.h"
 #include "Widgets/Views/SHeaderRow.h"
 #include "Widgets/Views/SListView.h"
@@ -42,7 +44,6 @@ namespace UE::MultiUserServer
 			.AutoHeight()
 			[
 				SNew(SPackageTransmissionTableFooter, PackageEntrySource.ToSharedRef())
-				.ExtendViewOptions(this, &SPackageTransmissionTable::ExtendViewOptions)
 				.TotalUnfilteredNum(InArgs._TotalUnfilteredNum)
 			]
 		];
@@ -61,13 +62,70 @@ namespace UE::MultiUserServer
 		PackageEntrySource->OnPackageEntriesAdded().RemoveAll(this);
 	}
 
+	TSharedRef<SWidget> SPackageTransmissionTable::CreateViewOptionsButton()
+	{
+		return SNew(SComboButton)
+			.ComboButtonStyle(FAppStyle::Get(), "SimpleComboButtonWithIcon")
+			.OnGetMenuContent_Lambda([this]()
+			{
+				FMenuBuilder MenuBuilder(true, nullptr);
+
+				MenuBuilder.AddMenuEntry(
+					LOCTEXT("DisplayTimestampInRelativeTime", "Display Relative Time"),
+					TAttribute<FText>::CreateLambda([this]()
+					{
+						const bool bIsVisible = HeaderRow->IsColumnVisible(SPackageTransmissionTableRow::TimeColumn);
+						return bIsVisible
+							? LOCTEXT("DisplayTimestampInRelativeTime.Tooltip.Visible", "Display the Last Modified column in relative time?")
+							: LOCTEXT("DisplayTimestampInRelativeTime.Tooltip.Hidden", "Disabled because the Timestamp column is hidden.");
+					}),
+					FSlateIcon(),
+					FUIAction(
+					FExecuteAction::CreateLambda([]()
+						{
+							UMultiUserServerPackageTransmissionSettings* Settings = UMultiUserServerPackageTransmissionSettings::GetSettings();
+				
+							switch (Settings->TimestampTimeFormat)
+							{
+							case ETimeFormat::Relative:
+								Settings->TimestampTimeFormat = ETimeFormat::Absolute;
+								break;
+							case ETimeFormat::Absolute: 
+								Settings->TimestampTimeFormat = ETimeFormat::Relative;
+								break;
+							default:
+								checkNoEntry();
+							}
+
+							Settings->SaveConfig();
+						}),
+						FCanExecuteAction::CreateLambda([this] { return HeaderRow->IsColumnVisible(SPackageTransmissionTableRow::TimeColumn); }),
+						FIsActionChecked::CreateLambda([this] { return UMultiUserServerPackageTransmissionSettings::GetSettings()->TimestampTimeFormat == ETimeFormat::Relative; })
+					),
+					NAME_None,
+					EUserInterfaceActionType::ToggleButton
+				);
+				
+				return MenuBuilder.MakeWidget();
+			})
+			.HasDownArrow(false)
+			.ButtonContent()
+			[
+				SNew(SImage)
+				.ColorAndOpacity(FSlateColor::UseForeground())
+				.Image(FAppStyle::Get().GetBrush("Icons.Settings"))
+			];
+	}
+
 	TSharedRef<SWidget> SPackageTransmissionTable::CreateTableView()
 	{
+		CreateHeaderRow();
 		return SAssignNew(TableView, SListView<TSharedPtr<FPackageTransmissionEntry>>)
 			.ListItemsSource(&PackageEntrySource->GetEntries())
+			.OnContextMenuOpening_Lambda([this]() { return ConcertSharedSlate::MakeTableContextMenu(HeaderRow.ToSharedRef(), GetDefaultColumnVisibilities()); })
 			.OnGenerateRow(this, &SPackageTransmissionTable::OnGenerateActivityRowWidget)
 			.SelectionMode(ESelectionMode::Multi)
-			.HeaderRow(CreateHeaderRow());
+			.HeaderRow(HeaderRow);
 	}
 
 	TSharedRef<SHeaderRow> SPackageTransmissionTable::CreateHeaderRow()
@@ -78,7 +136,7 @@ namespace UE::MultiUserServer
 			if (!bIsUpdatingColumnVisibility)
 			{
 				UMultiUserServerColumnVisibilitySettings::GetSettings()->SetPackageTransmissionColumnVisibility(
-				UE::ConcertSharedSlate::SnapshotColumnVisibilityState(HeaderRow.ToSharedRef())
+				ConcertSharedSlate::SnapshotColumnVisibilityState(HeaderRow.ToSharedRef())
 				);
 			}
 		});
@@ -96,13 +154,6 @@ namespace UE::MultiUserServer
 			{
 				Args.ShouldGenerateWidget(bCannotHide);
 			}
-			else
-			{
-				Args.OnGetMenuContent_Lambda([this]()
-				{
-					return ConcertSharedSlate::MakeHideColumnContextMenu(HeaderRow.ToSharedRef(), SPackageTransmissionTableRow::TimeColumn);
-				});
-			}
 			
 			HeaderRow->AddColumn(Args);
 		}
@@ -127,65 +178,29 @@ namespace UE::MultiUserServer
 		ConcertSharedSlate::RestoreColumnVisibilityState(HeaderRow.ToSharedRef(), ColumnSnapshot);
 	}
 
-	void SPackageTransmissionTable::ExtendViewOptions(FMenuBuilder& MenuBuilder)
+	void SPackageTransmissionTable::RestoreDefaultColumnVisibilities() const
 	{
-		MenuBuilder.AddMenuEntry(
-			LOCTEXT("SelectAll", "Show all"),
-			FText::GetEmpty(),
-			FSlateIcon(),
-			FUIAction(
-				FExecuteAction::CreateLambda([this]()
-				{
-					for(FName ColumnId : SPackageTransmissionTableRow::AllColumns)
-					{
-						HeaderRow->SetShowGeneratedColumn(ColumnId, true);
-					}
-				}),
-				FCanExecuteAction::CreateLambda([] { return true; })),
-			NAME_None,
-			EUserInterfaceActionType::Button
-		);
-		MenuBuilder.AddMenuEntry(
-			LOCTEXT("HideAll", "Hide all"),
-			FText::GetEmpty(),
-			FSlateIcon(),
-			FUIAction(
-				FExecuteAction::CreateLambda([this]()
-				{
-					for(FName ColumnId : SPackageTransmissionTableRow::AllColumns)
-					{
-						HeaderRow->SetShowGeneratedColumn(ColumnId, false);
-					}
-				}),
-				FCanExecuteAction::CreateLambda([] { return true; })),
-			NAME_None,
-			EUserInterfaceActionType::Button
-		);
-		MenuBuilder.AddMenuEntry(
-				LOCTEXT("RestoreDefaultColumnVisibility", "Restore columns visibility"),
-				FText::GetEmpty(),
-				FSlateIcon(),
-				FUIAction(
-					FExecuteAction::CreateSP(this, &SPackageTransmissionTable::RestoreDefaultColumnVisibilities),
-					FCanExecuteAction::CreateLambda([] { return true; })),
-				NAME_None,
-				EUserInterfaceActionType::Button
-			);
-		MenuBuilder.AddSeparator();
-		ConcertSharedSlate::AddEntriesForShowingHiddenRows(HeaderRow.ToSharedRef(), MenuBuilder);
+		const TMap<FName, bool> DefaultVisibilities = GetDefaultColumnVisibilities();
+		for (auto DefaultVisibilitiesIt = DefaultVisibilities.CreateConstIterator(); DefaultVisibilitiesIt; ++DefaultVisibilitiesIt)
+		{
+			HeaderRow->SetShowGeneratedColumn(DefaultVisibilitiesIt->Key, DefaultVisibilitiesIt->Value);
+		}
 	}
-
-	void SPackageTransmissionTable::RestoreDefaultColumnVisibilities()
+	
+	TMap<FName, bool> SPackageTransmissionTable::GetDefaultColumnVisibilities() const
 	{
+		TMap<FName, bool> Result;
 		const TSet<FName> HiddenByDefault = {
 			SPackageTransmissionTableRow::OriginColumn,
 			SPackageTransmissionTableRow::DestinationColumn,
 			SPackageTransmissionTableRow::PackagePathColumn
 		};
+		
 		for (FName ColumnName : SPackageTransmissionTableRow::AllColumns)
 		{
-			HeaderRow->SetShowGeneratedColumn(ColumnName, !HiddenByDefault.Contains(ColumnName));
+			Result.Add(ColumnName, !HiddenByDefault.Contains(ColumnName));
 		}
+		return Result;
 	}
 
 	void SPackageTransmissionTable::OnPackageEntriesModified(const TSet<FPackageTransmissionId>& Set) const
diff --git a/Engine/Plugins/Developer/Concert/ConcertApp/MultiUserServer/Source/MultiUserServer/Private/Widgets/Clients/PackageTransmission/SPackageTransmissionTable.h b/Engine/Plugins/Developer/Concert/ConcertApp/MultiUserServer/Source/MultiUserServer/Private/Widgets/Clients/PackageTransmission/SPackageTransmissionTable.h
--- a/Engine/Plugins/Developer/Concert/ConcertApp/MultiUserServer/Source/MultiUserServer/Private/Widgets/Clients/PackageTransmission/SPackageTransmissionTable.h
+++ b/Engine/Plugins/Developer/Concert/ConcertApp/MultiUserServer/Source/MultiUserServer/Private/Widgets/Clients/PackageTransmission/SPackageTransmissionTable.h
@@ -36,6 +36,8 @@ namespace UE::MultiUserServer
 
 		void Construct(const FArguments& InArgs, TSharedRef<IPackageTransmissionEntrySource> InPackageEntrySource, TSharedRef<FPackageTransmissionEntryTokenizer> InTokenizer);
 		virtual ~SPackageTransmissionTable() override;
+
+		TSharedRef<SWidget> CreateViewOptionsButton();
 		
 	private:
 
@@ -57,8 +59,8 @@ namespace UE::MultiUserServer
 		TSharedRef<ITableRow> OnGenerateActivityRowWidget(TSharedPtr<FPackageTransmissionEntry> Item, const TSharedRef<STableViewBase>& OwnerTable) const;
 		
 		void OnColumnVisibilitySettingsChanged(const FColumnVisibilitySnapshot& ColumnSnapshot);
-		void ExtendViewOptions(FMenuBuilder& MenuBuilder);
-		void RestoreDefaultColumnVisibilities(); 
+		void RestoreDefaultColumnVisibilities() const;
+		TMap<FName, bool> GetDefaultColumnVisibilities() const;
 		
 		void OnPackageEntriesModified(const TSet<FPackageTransmissionId>& Set) const;
 		void OnPackageArrayChanged(uint32 NumAdded) const;
diff --git a/Engine/Plugins/Developer/Concert/ConcertApp/MultiUserServer/Source/MultiUserServer/Private/Widgets/Clients/PackageTransmission/SPackageTransmissionTableFooter.cpp b/Engine/Plugins/Developer/Concert/ConcertApp/MultiUserServer/Source/MultiUserServer/Private/Widgets/Clients/PackageTransmission/SPackageTransmissionTableFooter.cpp
--- a/Engine/Plugins/Developer/Concert/ConcertApp/MultiUserServer/Source/MultiUserServer/Private/Widgets/Clients/PackageTransmission/SPackageTransmissionTableFooter.cpp
+++ b/Engine/Plugins/Developer/Concert/ConcertApp/MultiUserServer/Source/MultiUserServer/Private/Widgets/Clients/PackageTransmission/SPackageTransmissionTableFooter.cpp
@@ -4,10 +4,6 @@
 
 #include "Model/IPackageTransmissionEntrySource.h"
 
-#include "Framework/MultiBox/MultiBoxBuilder.h"
-#include "Styling/AppStyle.h"
-#include "Widgets/Images/SImage.h"
-#include "Widgets/Input/SComboButton.h"
 #include "Widgets/Layout/SSpacer.h"
 #include "Widgets/SBoxPanel.h"
 #include "Widgets/Text/STextBlock.h"
@@ -48,44 +44,8 @@ namespace UE::MultiUserServer
 			[
 				SNew(SSpacer)
 			]
-
-			+SHorizontalBox::Slot()
-			.AutoWidth()
-			.VAlign(VAlign_Center)
-			[
-				SNew(SComboButton)
-					.OnGetMenuContent(this, &SPackageTransmissionTableFooter::MakeViewOptionsMenuWidget, InArgs._ExtendViewOptions)
-					.HasDownArrow(true)
-					.ButtonContent()
-					[
-						SNew(SHorizontalBox)
-
-						+SHorizontalBox::Slot()
-						.AutoWidth()
-						.VAlign(VAlign_Center)
-						[
-							SNew(SImage).Image(FAppStyle::Get().GetBrush("Icons.Visible")) // The eye ball image.
-						]
-
-						+SHorizontalBox::Slot()
-						.AutoWidth()
-						.Padding(2, 0, 0, 0)
-						.VAlign(VAlign_Center)
-						[
-							SNew(STextBlock).Text(LOCTEXT("ViewOptions", "View Options"))
-						]
-					]
-			]
 		];
 	}
-
-	TSharedRef<SWidget> SPackageTransmissionTableFooter::MakeViewOptionsMenuWidget(FExtendContextMenu ExtendViewOptions)
-	{
-		FMenuBuilder MenuBuilder(/*bInShouldCloseWindowAfterMenuSelection=*/true, nullptr);
-
-		ExtendViewOptions.Execute(MenuBuilder);
-		return MenuBuilder.MakeWidget();
-	}
 }
 
 #undef LOCTEXT_NAMESPACE
\ No newline at end of file
diff --git a/Engine/Plugins/Developer/Concert/ConcertApp/MultiUserServer/Source/MultiUserServer/Private/Widgets/Clients/PackageTransmission/SPackageTransmissionTableFooter.h b/Engine/Plugins/Developer/Concert/ConcertApp/MultiUserServer/Source/MultiUserServer/Private/Widgets/Clients/PackageTransmission/SPackageTransmissionTableFooter.h
--- a/Engine/Plugins/Developer/Concert/ConcertApp/MultiUserServer/Source/MultiUserServer/Private/Widgets/Clients/PackageTransmission/SPackageTransmissionTableFooter.h
+++ b/Engine/Plugins/Developer/Concert/ConcertApp/MultiUserServer/Source/MultiUserServer/Private/Widgets/Clients/PackageTransmission/SPackageTransmissionTableFooter.h
@@ -32,7 +32,5 @@ namespace UE::MultiUserServer
 		
 		TSharedPtr<IPackageTransmissionEntrySource> PackageEntrySource;
 		TAttribute<uint32> TotalUnfilteredNum;
-		
-		TSharedRef<SWidget> MakeViewOptionsMenuWidget(FExtendContextMenu ExtendViewOptions);
 	};
 }
diff --git a/Engine/Plugins/Developer/Concert/ConcertApp/MultiUserServer/Source/MultiUserServer/Private/Widgets/Clients/PackageTransmission/SPackageTransmissionView.cpp b/Engine/Plugins/Developer/Concert/ConcertApp/MultiUserServer/Source/MultiUserServer/Private/Widgets/Clients/PackageTransmission/SPackageTransmissionView.cpp
--- a/Engine/Plugins/Developer/Concert/ConcertApp/MultiUserServer/Source/MultiUserServer/Private/Widgets/Clients/PackageTransmission/SPackageTransmissionView.cpp
+++ b/Engine/Plugins/Developer/Concert/ConcertApp/MultiUserServer/Source/MultiUserServer/Private/Widgets/Clients/PackageTransmission/SPackageTransmissionView.cpp
@@ -9,6 +9,8 @@
 
 #include "Widgets/SBoxPanel.h"
 
+#define LOCTEXT_NAMESPACE "UnrealMultiUserUI.SPackageTransmissionView"
+
 namespace UE::MultiUserServer
 {
 	void SPackageTransmissionView::Construct(const FArguments& InArgs, TSharedRef<IPackageTransmissionEntrySource> InPackageEntrySource, TSharedRef<FPackageTransmissionEntryTokenizer> InTokenizer)
@@ -16,6 +18,12 @@ namespace UE::MultiUserServer
 		PackageEntrySource = MoveTemp(InPackageEntrySource);
 		RootFilter = MakeFilter(InTokenizer);
 		FilteredModel = MakeShared<FFilteredPackageTransmissionModel>(PackageEntrySource.ToSharedRef(), RootFilter.ToSharedRef());
+
+		const TSharedRef<SPackageTransmissionTable> Table = SNew(SPackageTransmissionTable, FilteredModel.ToSharedRef(), InTokenizer)
+			.HighlightText_Lambda([this](){ return RootFilter->GetTextSearchFilter()->GetSearchText(); })
+			.CanScrollToLog(InArgs._CanScrollToLog)
+			.ScrollToLog(InArgs._ScrollToLog)
+			.TotalUnfilteredNum_Lambda([this](){ return PackageEntrySource->GetEntries().Num(); });
 		
 		ChildSlot
 		[
@@ -29,19 +37,22 @@ namespace UE::MultiUserServer
 				+SVerticalBox::Slot()
 				.AutoHeight()
 				[
-					RootFilter->BuildFilterWidgets()
+					RootFilter->BuildFilterWidgets
+					(
+						FFilterWidgetArgs().RightOfSearchBar
+						(
+							Table->CreateViewOptionsButton()
+						)
+					)
 				]
 				
 				+SVerticalBox::Slot()
 				[
-					SNew(SPackageTransmissionTable, FilteredModel.ToSharedRef(), InTokenizer)
-					.HighlightText_Lambda([this](){ return RootFilter->GetTextSearchFilter()->GetSearchText(); })
-					.CanScrollToLog(InArgs._CanScrollToLog)
-					.ScrollToLog(InArgs._ScrollToLog)
-					.TotalUnfilteredNum_Lambda([this](){ return PackageEntrySource->GetEntries().Num(); })
+					Table
 				]
 			]
 		];
 	}
 }
 
+#undef LOCTEXT_NAMESPACE 
\ No newline at end of file
diff --git a/Engine/Plugins/Developer/Concert/ConcertApp/MultiUserServer/Source/MultiUserServer/Private/Widgets/Clients/PackageTransmission/SPackageTransmissionView.h b/Engine/Plugins/Developer/Concert/ConcertApp/MultiUserServer/Source/MultiUserServer/Private/Widgets/Clients/PackageTransmission/SPackageTransmissionView.h
--- a/Engine/Plugins/Developer/Concert/ConcertApp/MultiUserServer/Source/MultiUserServer/Private/Widgets/Clients/PackageTransmission/SPackageTransmissionView.h
+++ b/Engine/Plugins/Developer/Concert/ConcertApp/MultiUserServer/Source/MultiUserServer/Private/Widgets/Clients/PackageTransmission/SPackageTransmissionView.h
@@ -38,5 +38,8 @@ namespace UE::MultiUserServer
 		
 		FCanScrollToLog CanScrollToLogDelegate;
 		FScrollToLog ScrollToLogDelegate;
+		
+		TSharedRef<SWidget> CreateOptionsButton();
+		TSharedRef<SWidget> CreateOptionsButtonMenu();
 	};
 }
diff --git a/Engine/Plugins/Developer/Concert/ConcertApp/MultiUserServer/Source/MultiUserServer/Private/Widgets/Clients/PackageTransmission/Util/PackageTransmissionEntryTokenizer.cpp b/Engine/Plugins/Developer/Concert/ConcertApp/MultiUserServer/Source/MultiUserServer/Private/Widgets/Clients/PackageTransmission/Util/PackageTransmissionEntryTokenizer.cpp
--- a/Engine/Plugins/Developer/Concert/ConcertApp/MultiUserServer/Source/MultiUserServer/Private/Widgets/Clients/PackageTransmission/Util/PackageTransmissionEntryTokenizer.cpp
+++ b/Engine/Plugins/Developer/Concert/ConcertApp/MultiUserServer/Source/MultiUserServer/Private/Widgets/Clients/PackageTransmission/Util/PackageTransmissionEntryTokenizer.cpp
@@ -16,7 +16,8 @@ namespace UE::MultiUserServer
 
 	FString FPackageTransmissionEntryTokenizer::TokenizeTime(const FPackageTransmissionEntry& Entry) const
 	{
-		return ConcertFrontendUtils::FormatTime(Entry.LocalStartTime, UMultiUserServerPackageTransmissionSettings::GetSettings()->TimestampTimeFormat).ToString();
+		const FDateTime LocalNow = FDateTime::Now();
+		return ConcertFrontendUtils::FormatTime(Entry.LocalStartTime, UMultiUserServerPackageTransmissionSettings::GetSettings()->TimestampTimeFormat, &LocalNow).ToString();
 	}
 
 	FString FPackageTransmissionEntryTokenizer::TokenizeOrigin(const FPackageTransmissionEntry& Entry) const
diff --git a/Engine/Plugins/Developer/Concert/ConcertApp/MultiUserServer/Source/MultiUserServer/Private/Widgets/Clients/SConcertClientsTabView.cpp b/Engine/Plugins/Developer/Concert/ConcertApp/MultiUserServer/Source/MultiUserServer/Private/Widgets/Clients/SConcertClientsTabView.cpp
--- a/Engine/Plugins/Developer/Concert/ConcertApp/MultiUserServer/Source/MultiUserServer/Private/Widgets/Clients/SConcertClientsTabView.cpp
+++ b/Engine/Plugins/Developer/Concert/ConcertApp/MultiUserServer/Source/MultiUserServer/Private/Widgets/Clients/SConcertClientsTabView.cpp
@@ -152,7 +152,7 @@ namespace UE::MultiUserServer
 		InTabManager->RegisterTabSpawner(ClientBrowserTabId, FOnSpawnTab::CreateSP(this, &SConcertClientsTabView::SpawnClientBrowserTab))
 			.SetDisplayName(LOCTEXT("ClientBrowserTabLabel", "Clients"));
 		InTabManager->RegisterTabSpawner(GlobalLogTabId, FOnSpawnTab::CreateSP(this, &SConcertClientsTabView::SpawnGlobalLogTab))
-			.SetDisplayName(LOCTEXT("GlobalLogTabLabel", "Global Log"));
+			.SetDisplayName(LOCTEXT("ServerLogTabLabel", "Server Log"));
 		InLayout->AddArea
 			(
 				FTabManager::NewPrimaryArea()
@@ -208,7 +208,7 @@ namespace UE::MultiUserServer
 	TSharedRef<SDockTab> SConcertClientsTabView::SpawnGlobalLogTab(const FSpawnTabArgs& InTabArgs)
 	{
 		return SNew(SDockTab)
-			.Label(LOCTEXT("GlobalLogTabLabel", "Global Log"))
+			.Label(LOCTEXT("ServerLogTabLabel", "Server Log"))
 			.TabRole(PanelTab)
 			[
 				SAssignNew(GlobalTransportLog, SConcertTransportLog, LogBuffer.ToSharedRef(), ClientInfoCache.ToSharedRef(), LogTokenizer.ToSharedRef())
@@ -220,7 +220,7 @@ namespace UE::MultiUserServer
 	{
 		return SNew(SButton)
 			.ButtonStyle(FAppStyle::Get(), "SimpleButton")
-			.ToolTipText(LOCTEXT("OpenGlobalLogTooltip", "Opens the Global Log which logs all incoming networked messages."))
+			.ToolTipText(LOCTEXT("OpenGlobalLogTooltip", "Opens the Server Log which logs all incoming networked messages."))
 			.ContentPadding(FMargin(1, 0))
 			.Visibility_Lambda(
 				[this]()
@@ -249,7 +249,7 @@ namespace UE::MultiUserServer
 					.Padding(4.0, 0.0f)
 					[
 						SNew(STextBlock)
-						.Text(LOCTEXT("OpenGlobalLog", "Open Global Log"))
+						.Text(LOCTEXT("OpenServerLog", "Open Server Log"))
 						.ColorAndOpacity(FSlateColor::UseForeground())
 					]
 			];
diff --git a/Engine/Plugins/Developer/Concert/ConcertApp/MultiUserServer/Source/MultiUserServer/Private/Widgets/Clients/SConcertClientsTabView.h b/Engine/Plugins/Developer/Concert/ConcertApp/MultiUserServer/Source/MultiUserServer/Private/Widgets/Clients/SConcertClientsTabView.h
--- a/Engine/Plugins/Developer/Concert/ConcertApp/MultiUserServer/Source/MultiUserServer/Private/Widgets/Clients/SConcertClientsTabView.h
+++ b/Engine/Plugins/Developer/Concert/ConcertApp/MultiUserServer/Source/MultiUserServer/Private/Widgets/Clients/SConcertClientsTabView.h
@@ -50,7 +50,7 @@ namespace UE::MultiUserServer
 
 		/** Whether it is possible to call ScrollToLog. */
 		bool CanScrollToLog(const FGuid& MessageId, FConcertLogEntryFilterFunc FilterFunc, FText& ErrorMessage) const;
-		/** Scrolls to the log in the global log tab. */
+		/** Scrolls to the log in the server log tab. */
 		void ScrollToLog(const FGuid& MessageId, FConcertLogEntryFilterFunc FilterFunc) const;
 		
 		bool IsGlobalLogOpen() const;
@@ -74,7 +74,7 @@ namespace UE::MultiUserServer
 
 		TSharedPtr<SConcertNetworkBrowser> ClientBrowser;
 
-		/** The widget inside of the global log tab */
+		/** The widget inside of the server log tab */
 		TSharedPtr<SConcertTransportLog> GlobalTransportLog;
 		
 		void CreateTabs(const TSharedRef<FTabManager>& InTabManager, const TSharedRef<FTabManager::FLayout>& InLayout, const FArguments& InArgs);
diff --git a/Engine/Plugins/Developer/Concert/ConcertApp/MultiUserServer/Source/MultiUserServer/Private/Widgets/Util/Filter/ConcertRootFrontendFilter.h b/Engine/Plugins/Developer/Concert/ConcertApp/MultiUserServer/Source/MultiUserServer/Private/Widgets/Util/Filter/ConcertRootFrontendFilter.h
--- a/Engine/Plugins/Developer/Concert/ConcertApp/MultiUserServer/Source/MultiUserServer/Private/Widgets/Util/Filter/ConcertRootFrontendFilter.h
+++ b/Engine/Plugins/Developer/Concert/ConcertApp/MultiUserServer/Source/MultiUserServer/Private/Widgets/Util/Filter/ConcertRootFrontendFilter.h
@@ -12,6 +12,15 @@
 
 namespace UE::MultiUserServer
 {
+	struct FFilterWidgetArgs
+	{
+		TSharedPtr<SWidget> _RightOfSearchBar;
+		FMargin _RightOfSearchBarPadding{ 4.f, 0.f, 4.f, 0.f };
+
+		FFilterWidgetArgs& RightOfSearchBar(TSharedPtr<SWidget> Widget) { _RightOfSearchBar = Widget; return *this; }
+		FFilterWidgetArgs& RightOfSearchBarPadding(FMargin Widget) { _RightOfSearchBarPadding = Widget; return *this; }
+	};
+	
 	/** A filter that contains multiple UI filters */
 	template<typename TFilterType, typename TTextSearchFilterType>
 	class TConcertFrontendRootFilter 
@@ -50,16 +59,17 @@ namespace UE::MultiUserServer
 				Filter->OnChanged().AddRaw(this, &TConcertFrontendRootFilter<TFilterType, TTextSearchFilterType>::BroadcastOnChanged);
 			}
 		}
-
+		
 		/** Builds the widget view for all contained filters */
-		TSharedRef<SWidget> BuildFilterWidgets()
+		TSharedRef<SWidget> BuildFilterWidgets(const FFilterWidgetArgs& InArgs = {})
 		{
-			return SNew(SVerticalBox)
+			TSharedPtr<SHorizontalBox> TopRow;
+			const TSharedRef<SVerticalBox> Result = SNew(SVerticalBox)
 
 				+SVerticalBox::Slot()
 				.AutoHeight()
 				[
-					SNew(SHorizontalBox)
+					SAssignNew(TopRow, SHorizontalBox)
 
 					// Filter button
 					+SHorizontalBox::Slot()
@@ -81,8 +91,19 @@ namespace UE::MultiUserServer
 				.AutoHeight()
 				[
 					FilterBar.ToSharedRef()
-				]
-			;
+				];
+
+			if (InArgs._RightOfSearchBar)
+			{
+				TopRow->AddSlot()
+				.AutoWidth()
+				.Padding(InArgs._RightOfSearchBarPadding)
+				[
+					InArgs._RightOfSearchBar.ToSharedRef()
+				];
+			}
+			
+			return Result;
 		}
 	
 		//~ Begin IFilter Interface
diff --git a/Engine/Plugins/Developer/Concert/ConcertUI/ConcertSharedSlate/Source/ConcertSharedSlate/Private/ConcertHeaderRowUtils.cpp b/Engine/Plugins/Developer/Concert/ConcertUI/ConcertSharedSlate/Source/ConcertSharedSlate/Private/ConcertHeaderRowUtils.cpp
--- a/Engine/Plugins/Developer/Concert/ConcertUI/ConcertSharedSlate/Source/ConcertSharedSlate/Private/ConcertHeaderRowUtils.cpp
+++ b/Engine/Plugins/Developer/Concert/ConcertUI/ConcertSharedSlate/Source/ConcertSharedSlate/Private/ConcertHeaderRowUtils.cpp
@@ -46,6 +46,86 @@ namespace UE::ConcertSharedSlate
 			);
 	}
 
+	TSharedRef<SWidget> MakeTableContextMenu(const TSharedRef<SHeaderRow>& HeaderRow, TMap<FName, bool> ColumnsVisibleByDefault, bool bDefaultVisibility)
+	{
+		FMenuBuilder MenuBuilder(true, nullptr);
+		AddDefaultControlEntries(HeaderRow, MenuBuilder, ColumnsVisibleByDefault, bDefaultVisibility);
+		MenuBuilder.AddSeparator();
+		AddEntriesForShowingHiddenRows(HeaderRow, MenuBuilder);
+		return MenuBuilder.MakeWidget();
+	}
+
+	void AddDefaultControlEntries(const TSharedRef<SHeaderRow>& HeaderRow, FMenuBuilder& MenuBuilder, TMap<FName, bool> ColumnsVisibleByDefault, bool bDefaultVisibility)
+	{
+		MenuBuilder.AddMenuEntry(
+			LOCTEXT("SelectAll", "Show all"),
+			FText::GetEmpty(),
+			FSlateIcon(),
+			FUIAction(
+				FExecuteAction::CreateLambda([WeakHeaderRow = TWeakPtr<SHeaderRow>(HeaderRow)]()
+				{
+					const TSharedPtr<SHeaderRow> HeaderRowPin = WeakHeaderRow.Pin();
+					if (!ensure(HeaderRowPin))
+					{
+						return;
+					}
+					
+					for (const SHeaderRow::FColumn& Column : HeaderRowPin->GetColumns())
+					{
+						HeaderRowPin->SetShowGeneratedColumn(Column.ColumnId, true);
+					}
+				}),
+				FCanExecuteAction::CreateLambda([] { return true; })),
+			NAME_None,
+			EUserInterfaceActionType::Button
+		);
+		MenuBuilder.AddMenuEntry(
+			LOCTEXT("HideAll", "Hide all"),
+			FText::GetEmpty(),
+			FSlateIcon(),
+			FUIAction(
+				FExecuteAction::CreateLambda([WeakHeaderRow = TWeakPtr<SHeaderRow>(HeaderRow)]()
+				{
+					const TSharedPtr<SHeaderRow> HeaderRowPin = WeakHeaderRow.Pin();
+					if (!ensure(HeaderRowPin))
+					{
+						return;
+					}
+					
+					for (const SHeaderRow::FColumn& Column : HeaderRowPin->GetColumns())
+					{
+						HeaderRowPin->SetShowGeneratedColumn(Column.ColumnId, false);
+					}
+				}),
+				FCanExecuteAction::CreateLambda([] { return true; })),
+			NAME_None,
+			EUserInterfaceActionType::Button
+		);
+		MenuBuilder.AddMenuEntry(
+			LOCTEXT("RestoreDefaultColumnVisibility", "Restore column visibility"),
+			FText::GetEmpty(),
+			FSlateIcon(),
+			FUIAction(
+				FExecuteAction::CreateLambda([WeakHeaderRow = TWeakPtr<SHeaderRow>(HeaderRow), ColumnsVisibleByDefault, bDefaultVisibility]
+				{
+					const TSharedPtr<SHeaderRow> HeaderRowPin = WeakHeaderRow.Pin();
+					if (!ensure(HeaderRowPin))
+					{
+						return;
+					}
+					
+					for (const SHeaderRow::FColumn& Column : HeaderRowPin->GetColumns())
+					{
+						const bool* IsVisible = ColumnsVisibleByDefault.Find(Column.ColumnId);
+						HeaderRowPin->SetShowGeneratedColumn(Column.ColumnId, IsVisible ? *IsVisible : bDefaultVisibility);
+					}
+				}),
+				FCanExecuteAction::CreateLambda([] { return true; })),
+			NAME_None,
+			EUserInterfaceActionType::Button
+		);
+	}
+
 	void AddEntriesForShowingHiddenRows(const TSharedRef<SHeaderRow>& HeaderRow, FMenuBuilder& MenuBuilder)
 	{
 		for (int32 ColumnIndex = 0; ColumnIndex < HeaderRow->GetColumns().Num(); ++ColumnIndex)
diff --git a/Engine/Plugins/Developer/Concert/ConcertUI/ConcertSharedSlate/Source/ConcertSharedSlate/Public/ConcertFrontendUtils.h b/Engine/Plugins/Developer/Concert/ConcertUI/ConcertSharedSlate/Source/ConcertSharedSlate/Public/ConcertFrontendUtils.h
--- a/Engine/Plugins/Developer/Concert/ConcertUI/ConcertSharedSlate/Source/ConcertSharedSlate/Public/ConcertFrontendUtils.h
+++ b/Engine/Plugins/Developer/Concert/ConcertUI/ConcertSharedSlate/Source/ConcertSharedSlate/Public/ConcertFrontendUtils.h
@@ -139,10 +139,10 @@ namespace ConcertFrontendUtils
 		return LOCTEXT("Now", "Now");
 	}
 
-	static FText FormatTime(const FDateTime& Time, ETimeFormat TimeFormat)
+	static FText FormatTime(const FDateTime& Time, ETimeFormat TimeFormat, const FDateTime* CurrTime = nullptr)
 	{
 		return TimeFormat == ETimeFormat::Relative
-			? FormatRelativeTime(Time)
+			? FormatRelativeTime(Time, CurrTime)
 			: FText::AsDateTime(Time);
 	}
 };
diff --git a/Engine/Plugins/Developer/Concert/ConcertUI/ConcertSharedSlate/Source/ConcertSharedSlate/Public/ConcertHeaderRowUtils.h b/Engine/Plugins/Developer/Concert/ConcertUI/ConcertSharedSlate/Source/ConcertSharedSlate/Public/ConcertHeaderRowUtils.h
--- a/Engine/Plugins/Developer/Concert/ConcertUI/ConcertSharedSlate/Source/ConcertSharedSlate/Public/ConcertHeaderRowUtils.h
+++ b/Engine/Plugins/Developer/Concert/ConcertUI/ConcertSharedSlate/Source/ConcertSharedSlate/Public/ConcertHeaderRowUtils.h
@@ -21,6 +21,8 @@ struct CONCERTSHAREDSLATE_API FColumnVisibilitySnapshot
 namespace UE::ConcertSharedSlate
 {
 	DECLARE_DELEGATE_OneParam(FSaveColumnVisibilitySnapshot, const FColumnVisibilitySnapshot& /*Snapshot*/);
+
+	// TODO DP: Remove MakeHideColumnContextMenu and AddHideColumnEntry once the UX pass is done
 	
 	/**
 	 * Creates a widget for hiding a SHeaderRow::FColumn when it is right-clicked.
@@ -28,6 +30,14 @@ namespace UE::ConcertSharedSlate
 	CONCERTSHAREDSLATE_API TSharedRef<SWidget> MakeHideColumnContextMenu(const TSharedRef<SHeaderRow>& HeaderRow, const FName ForColumnID);
 	CONCERTSHAREDSLATE_API void AddHideColumnEntry(const TSharedRef<SHeaderRow>& HeaderRow, const FName ForColumnID, FMenuBuilder& MenuBuilder);
 
+	/** Util function that can be directly fed into SListView::OnContextMenuOpening */
+	CONCERTSHAREDSLATE_API TSharedRef<SWidget> MakeTableContextMenu(const TSharedRef<SHeaderRow>& HeaderRow, TMap<FName, bool> ColumnsVisibleByDefault, bool bDefaultVisibility = true);
+
+	/**
+	 * Adds an entry for "Show all", "Hide all", and "Restore column visibility".
+	 */
+	CONCERTSHAREDSLATE_API void AddDefaultControlEntries(const TSharedRef<SHeaderRow>& HeaderRow, FMenuBuilder& MenuBuilder, TMap<FName, bool> ColumnsVisibleByDefault, bool bDefaultVisibility = true);
+
 	/**
 	 * Inspects the hidden rows on the header row and an entry for showing each hidden column.
 	 * @param HeaderRow To retrieve column information
