Add window menu and persist editor layout

diff --git a/Engine/Plugins/Developer/Concert/ConcertApp/MultiUserServer/Source/MultiUserServer/Private/Widgets/Browser/ConcertServerSessionBrowserController.cpp b/Engine/Plugins/Developer/Concert/ConcertApp/MultiUserServer/Source/MultiUserServer/Private/Widgets/Browser/ConcertServerSessionBrowserController.cpp
--- a/Engine/Plugins/Developer/Concert/ConcertApp/MultiUserServer/Source/MultiUserServer/Private/Widgets/Browser/ConcertServerSessionBrowserController.cpp
+++ b/Engine/Plugins/Developer/Concert/ConcertApp/MultiUserServer/Source/MultiUserServer/Private/Widgets/Browser/ConcertServerSessionBrowserController.cpp
@@ -42,7 +42,7 @@ namespace UE::MultiUserServer
 		
 		FGlobalTabmanager::Get()->RegisterTabSpawner(
 				ConcertServerTabs::GetSessionBrowserTabId(),
-				FOnSpawnTab::CreateRaw(this, &FConcertServerSessionBrowserController::SpawnSessionBrowserTab)
+				FOnSpawnTab::CreateRaw(this, &FConcertServerSessionBrowserController::SpawnSessionBrowserTab, Params.WindowController->GetRootWindow())
 			)
 			.SetDisplayName(LOCTEXT("SessionBrowserTabTitleLong", "Session Browser"))
 			.SetTooltipText(LOCTEXT("SessionBrowserTooltipText", "A section to browse, start, archive, and restore server sessions."))
@@ -247,17 +247,20 @@ namespace UE::MultiUserServer
 		}
 	}
 
-	TSharedRef<SDockTab> FConcertServerSessionBrowserController::SpawnSessionBrowserTab(const FSpawnTabArgs& Args)
+	TSharedRef<SDockTab> FConcertServerSessionBrowserController::SpawnSessionBrowserTab(const FSpawnTabArgs& Args, TSharedPtr<SWindow> RootWindow)
 	{
 		const TSharedRef<SDockTab> DockTab = SNew(SDockTab)
 			.Label(LOCTEXT("SessionBrowserTabTitle", "Sessions"))
 			.TabRole(MajorTab)
-			.OnCanCloseTab(SDockTab::FCanCloseTab::CreateLambda([](){ return false; }))
-			[
-				SAssignNew(ConcertBrowser, SConcertServerSessionBrowser, SharedThis(this))
-					.DoubleClickLiveSession(this, &FConcertServerSessionBrowserController::OpenSession)
-					.DoubleClickArchivedSession(this, &FConcertServerSessionBrowserController::OpenSession)
-			];
+			.CanEverClose(false);
+
+		DockTab->SetContent(
+			SAssignNew(ConcertBrowser, SConcertServerSessionBrowser, SharedThis(this))
+				.ConstructUnderMajorTab(DockTab)
+				.ConstructUnderWindow(RootWindow)
+				.DoubleClickLiveSession(this, &FConcertServerSessionBrowserController::OpenSession)
+				.DoubleClickArchivedSession(this, &FConcertServerSessionBrowserController::OpenSession)
+			);
 
 		return DockTab;
 	}
diff --git a/Engine/Plugins/Developer/Concert/ConcertApp/MultiUserServer/Source/MultiUserServer/Private/Widgets/Browser/ConcertServerSessionBrowserController.h b/Engine/Plugins/Developer/Concert/ConcertApp/MultiUserServer/Source/MultiUserServer/Private/Widgets/Browser/ConcertServerSessionBrowserController.h
--- a/Engine/Plugins/Developer/Concert/ConcertApp/MultiUserServer/Source/MultiUserServer/Private/Widgets/Browser/ConcertServerSessionBrowserController.h
+++ b/Engine/Plugins/Developer/Concert/ConcertApp/MultiUserServer/Source/MultiUserServer/Private/Widgets/Browser/ConcertServerSessionBrowserController.h
@@ -64,7 +64,7 @@ namespace UE::MultiUserServer
 
 		TSharedPtr<SConcertServerSessionBrowser> ConcertBrowser;
 
-		TSharedRef<SDockTab> SpawnSessionBrowserTab(const FSpawnTabArgs& Args);
+		TSharedRef<SDockTab> SpawnSessionBrowserTab(const FSpawnTabArgs& Args, TSharedPtr<SWindow> RootWindow);
 
 		// Update view when session list changes
 		void OnLiveSessionCreated(bool, const IConcertServer&, TSharedRef<IConcertServerSession>) { RefreshSessionList(); }
diff --git a/Engine/Plugins/Developer/Concert/ConcertApp/MultiUserServer/Source/MultiUserServer/Private/Widgets/Browser/SConcertServerSessionBrowser.cpp b/Engine/Plugins/Developer/Concert/ConcertApp/MultiUserServer/Source/MultiUserServer/Private/Widgets/Browser/SConcertServerSessionBrowser.cpp
--- a/Engine/Plugins/Developer/Concert/ConcertApp/MultiUserServer/Source/MultiUserServer/Private/Widgets/Browser/SConcertServerSessionBrowser.cpp
+++ b/Engine/Plugins/Developer/Concert/ConcertApp/MultiUserServer/Source/MultiUserServer/Private/Widgets/Browser/SConcertServerSessionBrowser.cpp
@@ -12,30 +12,65 @@
 #include "Window/ModalWindowManager.h"
 
 #include "Dialog/SMessageDialog.h"
+#include "Widgets/Docking/SDockTab.h"
 
 #define LOCTEXT_NAMESPACE "UnrealMultiUserUI.SConcertServerSessionBrowser"
 
 namespace UE::MultiUserServer
 {
+	const FName SConcertServerSessionBrowser::SessionBrowserTabId("SessionBrowserTabId");
+
 	void SConcertServerSessionBrowser::Construct(const FArguments& InArgs, TSharedRef<FConcertServerSessionBrowserController> InController)
 	{
 		Controller = InController;
-		SConcertTabViewBase::Construct(
-			SConcertTabViewBase::FArguments()
-			.Content()
-			[
-				MakeSessionTableView(InArgs)
-			],
+		SConcertTabViewWithManagerBase::Construct(
+			SConcertTabViewWithManagerBase::FArguments()
+			.ConstructUnderWindow(InArgs._ConstructUnderWindow)
+			.ConstructUnderMajorTab(InArgs._ConstructUnderMajorTab)
+			.CreateTabs(FCreateTabs::CreateLambda([this, &InArgs](const TSharedRef<FTabManager>& InTabManager, const TSharedRef<FTabManager::FLayout>& InLayout)
+			{
+				CreateTabs(InTabManager, InLayout, InArgs);
+			}))
+			.LayoutName("ConcertSessionBrowserTabView_v0.1"),
 			ConcertServerTabs::GetSessionBrowserTabId()
 			);
 	}
 
-	TSharedRef<SWidget> SConcertServerSessionBrowser::MakeSessionTableView(const FArguments& InArgs)
+	void SConcertServerSessionBrowser::CreateTabs(const TSharedRef<FTabManager>& InTabManager, const TSharedRef<FTabManager::FLayout>& InLayout, const FArguments& InArgs)
+	{
+		const TSharedRef<FWorkspaceItem> WorkspaceItem = InTabManager->AddLocalWorkspaceMenuCategory(LOCTEXT("SessionBrowser", "Session Browser"));
+		
+		InTabManager->RegisterTabSpawner(SessionBrowserTabId, FOnSpawnTab::CreateSP(this, &SConcertServerSessionBrowser::SpawnSessionBrowserTab, InArgs._DoubleClickLiveSession, InArgs._DoubleClickArchivedSession))
+			.SetDisplayName(LOCTEXT("SessionBrowserTabLabel", "Sessions"))
+			.SetGroup(WorkspaceItem);
+
+		InLayout->AddArea
+		(
+			FTabManager::NewPrimaryArea()->Split
+			(
+				FTabManager::NewStack()
+				->AddTab(SessionBrowserTabId, ETabState::OpenedTab)
+				->SetHideTabWell(true)
+			)
+		);
+	}
+
+	TSharedRef<SDockTab> SConcertServerSessionBrowser::SpawnSessionBrowserTab(const FSpawnTabArgs& InTabArgs, FSessionDelegate DoubleClickLiveSession, FSessionDelegate DoubleClickArchivedSession)
+	{
+		return SNew(SDockTab)
+			.Label(LOCTEXT("SessionBrowserTabLabel", "Sessions"))
+			.TabRole(PanelTab)
+			[
+				MakeSessionTableView(MoveTemp(DoubleClickLiveSession), MoveTemp(DoubleClickArchivedSession))
+			]; 
+	}
+
+	TSharedRef<SWidget> SConcertServerSessionBrowser::MakeSessionTableView(FSessionDelegate DoubleClickLiveSession, FSessionDelegate DoubleClickArchivedSession)
 	{
 		SearchText = MakeShared<FText>();
 		return SAssignNew(SessionBrowser, SConcertSessionBrowser, Controller.Pin().ToSharedRef(), SearchText)
-			.OnLiveSessionDoubleClicked(InArgs._DoubleClickLiveSession)
-			.OnArchivedSessionDoubleClicked(InArgs._DoubleClickArchivedSession)
+			.OnLiveSessionDoubleClicked(MoveTemp(DoubleClickLiveSession))
+			.OnArchivedSessionDoubleClicked(MoveTemp(DoubleClickArchivedSession))
 			.PostRequestedDeleteSession(this, &SConcertServerSessionBrowser::RequestDeleteSession)
 			// Pretend a modal dialog said no - RequestDeleteSession will show non-modal dialog
 			.AskUserToDeleteSessions_Lambda([](auto) { return false; })
diff --git a/Engine/Plugins/Developer/Concert/ConcertApp/MultiUserServer/Source/MultiUserServer/Private/Widgets/Browser/SConcertServerSessionBrowser.h b/Engine/Plugins/Developer/Concert/ConcertApp/MultiUserServer/Source/MultiUserServer/Private/Widgets/Browser/SConcertServerSessionBrowser.h
--- a/Engine/Plugins/Developer/Concert/ConcertApp/MultiUserServer/Source/MultiUserServer/Private/Widgets/Browser/SConcertServerSessionBrowser.h
+++ b/Engine/Plugins/Developer/Concert/ConcertApp/MultiUserServer/Source/MultiUserServer/Private/Widgets/Browser/SConcertServerSessionBrowser.h
@@ -5,7 +5,7 @@
 #include "CoreMinimal.h"
 #include "Session/Browser/SConcertSessionBrowser.h"
 #include "Widgets/DeclarativeSyntaxSupport.h"
-#include "Widgets/SConcertTabViewBase.h"
+#include "Widgets/SConcertTabViewWithManagerBase.h"
 
 class SConcertSessionBrowser;
 
@@ -14,16 +14,20 @@ namespace UE::MultiUserServer
 	class FConcertServerSessionBrowserController;
 
 	/** Shows a list of server sessions */
-	class SConcertServerSessionBrowser : public SConcertTabViewBase
+	class SConcertServerSessionBrowser : public SConcertTabViewWithManagerBase
 	{
 	public:
+		
+		static const FName SessionBrowserTabId;
 	
 		SLATE_BEGIN_ARGS(SConcertServerSessionBrowser) { }
-		SLATE_EVENT(FSessionDelegate, DoubleClickLiveSession)
-		SLATE_EVENT(FSessionDelegate, DoubleClickArchivedSession)
-	SLATE_END_ARGS()
+			SLATE_ARGUMENT(TSharedPtr<SDockTab>, ConstructUnderMajorTab)
+			SLATE_ARGUMENT(TSharedPtr<SWindow>, ConstructUnderWindow)
+			SLATE_EVENT(FSessionDelegate, DoubleClickLiveSession)
+			SLATE_EVENT(FSessionDelegate, DoubleClickArchivedSession)
+		SLATE_END_ARGS()
 
-	void Construct(const FArguments& InArgs, TSharedRef<FConcertServerSessionBrowserController> InController);
+		void Construct(const FArguments& InArgs, TSharedRef<FConcertServerSessionBrowserController> InController);
 
 		void RequestRefreshListNextTick() { bRequestedRefresh = true; }
 		virtual void Tick(const FGeometry& AllottedGeometry, const double InCurrentTime, const float InDeltaTime) override
@@ -46,9 +50,12 @@ namespace UE::MultiUserServer
 
 		TSharedPtr<FText> SearchText;
 		TSharedPtr<SConcertSessionBrowser> SessionBrowser;
-	
-		TSharedRef<SWidget> MakeSessionTableView(const FArguments& InArgs);
-
+		
+		void CreateTabs(const TSharedRef<FTabManager>& InTabManager, const TSharedRef<FTabManager::FLayout>& InLayout, const FArguments& InArgs);
+		
+		TSharedRef<SDockTab> SpawnSessionBrowserTab(const FSpawnTabArgs& InTabArgs, FSessionDelegate DoubleClickLiveSession, FSessionDelegate DoubleClickArchivedSession);
+		TSharedRef<SWidget> MakeSessionTableView(FSessionDelegate DoubleClickLiveSession, FSessionDelegate DoubleClickArchivedSession);
+		
 		void RequestDeleteSession(const TArray<TSharedPtr<FConcertSessionTreeItem>>& SessionItems);
 		void DeleteSessionsWithFakeModalQuestion(const FText& Message, const TArray<TSharedPtr<FConcertSessionTreeItem>>& SessionItems);
 	};
diff --git a/Engine/Plugins/Developer/Concert/ConcertApp/MultiUserServer/Source/MultiUserServer/Private/Widgets/Clients/PackageTransmission/PackageTransmissionTabController.cpp b/Engine/Plugins/Developer/Concert/ConcertApp/MultiUserServer/Source/MultiUserServer/Private/Widgets/Clients/PackageTransmission/PackageTransmissionTabController.cpp
--- a/Engine/Plugins/Developer/Concert/ConcertApp/MultiUserServer/Source/MultiUserServer/Private/Widgets/Clients/PackageTransmission/PackageTransmissionTabController.cpp
+++ b/Engine/Plugins/Developer/Concert/ConcertApp/MultiUserServer/Source/MultiUserServer/Private/Widgets/Clients/PackageTransmission/PackageTransmissionTabController.cpp
@@ -16,6 +16,7 @@ namespace UE::MultiUserServer
 	FPackageTransmissionTabController::FPackageTransmissionTabController(
 			FName TabId,
 			TSharedRef<FTabManager> OwningTabManager,
+			TSharedRef<FWorkspaceItem> WorkspaceItem,
 			TSharedRef<FPackageTransmissionModel> TransmissionModel,
 			TSharedRef<FEndpointToUserNameCache> EndpointToUserNameCache,
 			FCanScrollToLog CanScrollToLogDelegate,
@@ -31,7 +32,8 @@ namespace UE::MultiUserServer
 	{
 		OwningTabManager->RegisterTabSpawner(TabId, FOnSpawnTab::CreateRaw(this, &FPackageTransmissionTabController::SpawnTab))
 			// In the future we may create multiple FPackageTransmissionTabController and may want to change this name to depend on some parameter
-			.SetDisplayName(LOCTEXT("PackageTabLabel", "Packages"));
+			.SetDisplayName(LOCTEXT("PackageTabLabel", "Packages"))
+			.SetGroup(WorkspaceItem);
 	}
 
 	FPackageTransmissionTabController::~FPackageTransmissionTabController()
@@ -42,7 +44,7 @@ namespace UE::MultiUserServer
 
 	TSharedRef<SDockTab> FPackageTransmissionTabController::SpawnTab(const FSpawnTabArgs& SpawnTabArgs)
 	{
-		return SAssignNew(DockTab, SDockTab)
+		return SNew(SDockTab)
 			.Label(LOCTEXT("PackageTabLabel", "Packages"))
 			.TabRole(PanelTab)
 			[
diff --git a/Engine/Plugins/Developer/Concert/ConcertApp/MultiUserServer/Source/MultiUserServer/Private/Widgets/Clients/PackageTransmission/PackageTransmissionTabController.h b/Engine/Plugins/Developer/Concert/ConcertApp/MultiUserServer/Source/MultiUserServer/Private/Widgets/Clients/PackageTransmission/PackageTransmissionTabController.h
--- a/Engine/Plugins/Developer/Concert/ConcertApp/MultiUserServer/Source/MultiUserServer/Private/Widgets/Clients/PackageTransmission/PackageTransmissionTabController.h
+++ b/Engine/Plugins/Developer/Concert/ConcertApp/MultiUserServer/Source/MultiUserServer/Private/Widgets/Clients/PackageTransmission/PackageTransmissionTabController.h
@@ -8,6 +8,7 @@
 class FEndpointToUserNameCache;
 class FSpawnTabArgs;
 class FTabManager;
+class FWorkspaceItem;
 class IConcertSyncServer;
 class SDockTab;
 
@@ -16,7 +17,7 @@ namespace UE::MultiUserServer
 	class FPackageTransmissionEntryTokenizer;
 	class FPackageTransmissionModel;
 
-	/** Manages a tab that displays a package transmission model. The model */
+	/** Manages a tab that displays a package transmission model */
 	class FPackageTransmissionTabController
 	{
 	public:
@@ -24,6 +25,7 @@ namespace UE::MultiUserServer
 		FPackageTransmissionTabController(
 			FName TabId,
 			TSharedRef<FTabManager> OwningTabManager,
+			TSharedRef<FWorkspaceItem> WorkspaceItem,
 			TSharedRef<FPackageTransmissionModel> TransmissionModel,
 			TSharedRef<FEndpointToUserNameCache> EndpointToUserNameCache,
 			FCanScrollToLog CanScrollToLogDelegate,
@@ -42,7 +44,6 @@ namespace UE::MultiUserServer
 		FScrollToLog ScrollToLogDelegate;
 		
 		TSharedRef<FPackageTransmissionEntryTokenizer> Tokenizer;
-		TSharedPtr<SDockTab> DockTab;
 
 		TSharedRef<SDockTab> SpawnTab(const FSpawnTabArgs& SpawnTabArgs);
 	};
diff --git a/Engine/Plugins/Developer/Concert/ConcertApp/MultiUserServer/Source/MultiUserServer/Private/Widgets/Clients/SConcertClientsTabView.cpp b/Engine/Plugins/Developer/Concert/ConcertApp/MultiUserServer/Source/MultiUserServer/Private/Widgets/Clients/SConcertClientsTabView.cpp
--- a/Engine/Plugins/Developer/Concert/ConcertApp/MultiUserServer/Source/MultiUserServer/Private/Widgets/Clients/SConcertClientsTabView.cpp
+++ b/Engine/Plugins/Developer/Concert/ConcertApp/MultiUserServer/Source/MultiUserServer/Private/Widgets/Clients/SConcertClientsTabView.cpp
@@ -140,19 +140,24 @@ namespace UE::MultiUserServer
 
 	void SConcertClientsTabView::CreateTabs(const TSharedRef<FTabManager>& InTabManager, const TSharedRef<FTabManager::FLayout>& InLayout, const FArguments& InArgs)
 	{
-		MainPackageTransmissionTab = MakeShared<UE::MultiUserServer::FPackageTransmissionTabController>(
+		const TSharedRef<FWorkspaceItem> WorkspaceItem = InTabManager->AddLocalWorkspaceMenuCategory(LOCTEXT("ClientBrowser", "Client Browser"));
+		
+		MainPackageTransmissionTab = MakeShared<FPackageTransmissionTabController>(
 			PackageTransmissionTabId,
 			InTabManager,
+			WorkspaceItem,
 			PackageTransmissionModel.ToSharedRef(),
 			ClientInfoCache.ToSharedRef(),
-			UE::MultiUserServer::FCanScrollToLog::CreateSP(this, &SConcertClientsTabView::CanScrollToLog),
-			UE::MultiUserServer::FScrollToLog::CreateSP(this, &SConcertClientsTabView::ScrollToLog)
+			FCanScrollToLog::CreateSP(this, &SConcertClientsTabView::CanScrollToLog),
+			FScrollToLog::CreateSP(this, &SConcertClientsTabView::ScrollToLog)
 			);
-		
 		InTabManager->RegisterTabSpawner(ClientBrowserTabId, FOnSpawnTab::CreateSP(this, &SConcertClientsTabView::SpawnClientBrowserTab))
-			.SetDisplayName(LOCTEXT("ClientBrowserTabLabel", "Clients"));
+			.SetDisplayName(LOCTEXT("ClientBrowserTabLabel", "Clients"))
+			.SetGroup(WorkspaceItem);
 		InTabManager->RegisterTabSpawner(GlobalLogTabId, FOnSpawnTab::CreateSP(this, &SConcertClientsTabView::SpawnGlobalLogTab))
-			.SetDisplayName(LOCTEXT("ServerLogTabLabel", "Server Log"));
+			.SetDisplayName(LOCTEXT("ServerLogTabLabel", "Server Log"))
+			.SetGroup(WorkspaceItem);
+		
 		InLayout->AddArea
 			(
 				FTabManager::NewPrimaryArea()
diff --git a/Engine/Plugins/Developer/Concert/ConcertApp/MultiUserServer/Source/MultiUserServer/Private/Widgets/SConcertTabViewWithManagerBase.cpp b/Engine/Plugins/Developer/Concert/ConcertApp/MultiUserServer/Source/MultiUserServer/Private/Widgets/SConcertTabViewWithManagerBase.cpp
--- a/Engine/Plugins/Developer/Concert/ConcertApp/MultiUserServer/Source/MultiUserServer/Private/Widgets/SConcertTabViewWithManagerBase.cpp
+++ b/Engine/Plugins/Developer/Concert/ConcertApp/MultiUserServer/Source/MultiUserServer/Private/Widgets/SConcertTabViewWithManagerBase.cpp
@@ -2,8 +2,12 @@
 
 #include "SConcertTabViewWithManagerBase.h"
 
+#include "Framework/MultiBox/MultiBoxBuilder.h"
+#include "Framework/Docking/LayoutService.h"
 #include "Framework/Docking/TabManager.h"
 
+#define LOCTEXT_NAMESPACE "UnrealMultiUserUI.SConcertTabViewWithManagerBase"
+
 void SConcertTabViewWithManagerBase::Construct(const FArguments& InArgs, FName InStatusBarId)
 {
 	check(InArgs._ConstructUnderWindow && InArgs._ConstructUnderMajorTab);
@@ -20,7 +24,48 @@ void SConcertTabViewWithManagerBase::Construct(const FArguments& InArgs, FName I
 TSharedRef<SWidget> SConcertTabViewWithManagerBase::CreateTabs(const FArguments& InArgs)
 {
 	TabManager = FGlobalTabmanager::Get()->NewTabManager(InArgs._ConstructUnderMajorTab.ToSharedRef());
-	const TSharedRef<FTabManager::FLayout> Layout = FTabManager::NewLayout(InArgs._LayoutName);
+	TabManager->SetMainTab(InArgs._ConstructUnderMajorTab.ToSharedRef());
+	
+	TSharedRef<FTabManager::FLayout> Layout = FTabManager::NewLayout(InArgs._LayoutName);
 	InArgs._CreateTabs.ExecuteIfBound(TabManager.ToSharedRef(), Layout);
-	return TabManager->RestoreFrom(Layout, InArgs._ConstructUnderWindow).ToSharedRef();
+	Layout = FLayoutSaveRestore::LoadFromConfig(GEditorLayoutIni, Layout);
+	TabManager->SetOnPersistLayout(
+		FTabManager::FOnPersistLayout::CreateStatic(
+			[](const TSharedRef<FTabManager::FLayout>& InLayout)
+			{
+				if (InLayout->GetPrimaryArea().Pin().IsValid())
+				{
+					FLayoutSaveRestore::SaveToConfig(GEditorLayoutIni, InLayout);
+				}
+			}
+		)
+	);
+	TSharedRef<SWidget> Result = TabManager->RestoreFrom(Layout, InArgs._ConstructUnderWindow).ToSharedRef();
+	
+	FMenuBarBuilder MenuBarBuilder = FMenuBarBuilder(TSharedPtr<FUICommandList>());
+	FillInDefaultMenuItems(MenuBarBuilder);
+	InArgs._CreateMenuBar.ExecuteIfBound(MenuBarBuilder);
+	
+	const TSharedRef<SWidget> MenuWidget = MenuBarBuilder.MakeWidget();
+	TabManager->SetAllowWindowMenuBar(true);
+	TabManager->SetMenuMultiBox(MenuBarBuilder.GetMultiBox(), MenuWidget);
+	
+	return Result;
+}
+
+void SConcertTabViewWithManagerBase::FillInDefaultMenuItems(FMenuBarBuilder MenuBarBuilder)
+{
+	MenuBarBuilder.AddPullDownMenu(
+			LOCTEXT("WindowMenuLabel", "Window"),
+			FText::GetEmpty(),
+			FNewMenuDelegate::CreateSP(this, &SConcertTabViewWithManagerBase::FillWindowMenu),
+			"Window"
+		);
 }
+
+void SConcertTabViewWithManagerBase::FillWindowMenu(FMenuBuilder& MenuBuilder)
+{
+	TabManager->PopulateLocalTabSpawnerMenu(MenuBuilder);
+}
+
+#undef LOCTEXT_NAMESPACE
diff --git a/Engine/Plugins/Developer/Concert/ConcertApp/MultiUserServer/Source/MultiUserServer/Private/Widgets/SConcertTabViewWithManagerBase.h b/Engine/Plugins/Developer/Concert/ConcertApp/MultiUserServer/Source/MultiUserServer/Private/Widgets/SConcertTabViewWithManagerBase.h
--- a/Engine/Plugins/Developer/Concert/ConcertApp/MultiUserServer/Source/MultiUserServer/Private/Widgets/SConcertTabViewWithManagerBase.h
+++ b/Engine/Plugins/Developer/Concert/ConcertApp/MultiUserServer/Source/MultiUserServer/Private/Widgets/SConcertTabViewWithManagerBase.h
@@ -7,12 +7,15 @@
 #include "Widgets/DeclarativeSyntaxSupport.h"
 #include "SConcertTabViewBase.h"
 
+class FMenuBarBuilder;
+
 /** Base class for tab views that create sub-tabs. */
 class SConcertTabViewWithManagerBase : public SConcertTabViewBase
 {
 public:
 
 	DECLARE_DELEGATE_TwoParams(FCreateTabs, const TSharedRef<FTabManager>& /*TabManager*/, const TSharedRef<FTabManager::FLayout>& /*Layout*/);
+	DECLARE_DELEGATE_OneParam(FCreateWindowMenu, FMenuBarBuilder&);
 	DECLARE_DELEGATE_RetVal_OneParam(TSharedRef<SWidget>, FOverlayTabs, const TSharedRef<SWidget>& /* Tabs */);
 	
 	SLATE_BEGIN_ARGS(SConcertTabViewWithManagerBase) {}
@@ -22,6 +25,8 @@ public:
 		SLATE_ARGUMENT(TSharedPtr<SWindow>, ConstructUnderWindow)
 		/** Callback for creating the sub-tabs */
 		SLATE_EVENT(FCreateTabs, CreateTabs)
+		/** Called to fill the major tab's menu bar */
+		SLATE_EVENT(FCreateWindowMenu, CreateMenuBar)
 		/** Optional function to overlay all the tabs (for example to overlay a message to execute a console command first) */
 		SLATE_EVENT(FOverlayTabs, OverlayTabs)
 		/** Name to give the layout. Important for saving config. */
@@ -44,4 +49,6 @@ private:
 	TSharedPtr<FTabManager> TabManager;
 	
 	TSharedRef<SWidget> CreateTabs(const FArguments& InArgs);
+	void FillInDefaultMenuItems(FMenuBarBuilder MenuBarBuilder);
+	void FillWindowMenu(FMenuBuilder& MenuBuilder);
 };
diff --git a/Engine/Plugins/Developer/Concert/ConcertApp/MultiUserServer/Source/MultiUserServer/Private/Widgets/SessionTabs/Archived/ArchivedConcertSessionTab.cpp b/Engine/Plugins/Developer/Concert/ConcertApp/MultiUserServer/Source/MultiUserServer/Private/Widgets/SessionTabs/Archived/ArchivedConcertSessionTab.cpp
--- a/Engine/Plugins/Developer/Concert/ConcertApp/MultiUserServer/Source/MultiUserServer/Private/Widgets/SessionTabs/Archived/ArchivedConcertSessionTab.cpp
+++ b/Engine/Plugins/Developer/Concert/ConcertApp/MultiUserServer/Source/MultiUserServer/Private/Widgets/SessionTabs/Archived/ArchivedConcertSessionTab.cpp
@@ -37,8 +37,10 @@ void FArchivedConcertSessionTab::CreateDockContent(const TSharedRef<SDockTab>& I
 {
 	SEditableSessionHistory::FMakeSessionHistory MakeSessionHistory = SEditableSessionHistory::FMakeSessionHistory::CreateLambda([this](SSessionHistory::FArguments Arguments)
 	{
-		checkf(!HistoryController.IsValid(), TEXT("Called more than once"));
-		HistoryController = UE::MultiUserServer::CreateForInspector(InspectedSessionID, SyncServer, MoveTemp(Arguments));
+		if (!HistoryController)
+		{
+			HistoryController = UE::MultiUserServer::CreateForInspector(InspectedSessionID, SyncServer, MoveTemp(Arguments));
+		}
 		return HistoryController->GetSessionHistory();
 	});
 	
diff --git a/Engine/Plugins/Developer/Concert/ConcertApp/MultiUserServer/Source/MultiUserServer/Private/Widgets/SessionTabs/Archived/SConcertArchivedSessionTabView.cpp b/Engine/Plugins/Developer/Concert/ConcertApp/MultiUserServer/Source/MultiUserServer/Private/Widgets/SessionTabs/Archived/SConcertArchivedSessionTabView.cpp
--- a/Engine/Plugins/Developer/Concert/ConcertApp/MultiUserServer/Source/MultiUserServer/Private/Widgets/SessionTabs/Archived/SConcertArchivedSessionTabView.cpp
+++ b/Engine/Plugins/Developer/Concert/ConcertApp/MultiUserServer/Source/MultiUserServer/Private/Widgets/SessionTabs/Archived/SConcertArchivedSessionTabView.cpp
@@ -28,6 +28,8 @@ void SConcertArchivedSessionTabView::Construct(const FArguments& InArgs, FName I
 
 void SConcertArchivedSessionTabView::CreateTabs(const TSharedRef<FTabManager>& InTabManager, const TSharedRef<FTabManager::FLayout>& InLayout, const FArguments& InArgs)
 {
+	const TSharedRef<FWorkspaceItem> WorkspaceItem = InTabManager->AddLocalWorkspaceMenuCategory(LOCTEXT("ArchiveSession", "Archived Session"));
+	
 	InTabManager->RegisterTabSpawner(HistoryTabId, FOnSpawnTab::CreateSP(this, &SConcertArchivedSessionTabView::SpawnActivityHistory,
 		InArgs._MakeSessionHistory,
 		InArgs._CanDeleteActivities,
@@ -37,7 +39,8 @@ void SConcertArchivedSessionTabView::CreateTabs(const TSharedRef<FTabManager>& I
 		InArgs._CanUnmuteActivities,
 		InArgs._UnmuteActivities
 		))
-		.SetDisplayName(LOCTEXT("ActivityHistoryLabel", "History"));
+		.SetDisplayName(LOCTEXT("ActivityHistoryLabel", "History"))
+		.SetGroup(WorkspaceItem);
 	InLayout->AddArea
 		(
 			FTabManager::NewPrimaryArea()
@@ -47,6 +50,7 @@ void SConcertArchivedSessionTabView::CreateTabs(const TSharedRef<FTabManager>& I
 					FTabManager::NewStack()
 					->SetSizeCoefficient(0.5f)
 					->AddTab(HistoryTabId, ETabState::OpenedTab)
+					->SetHideTabWell(true)
 				)
 		);
 }
diff --git a/Engine/Plugins/Developer/Concert/ConcertApp/MultiUserServer/Source/MultiUserServer/Private/Widgets/SessionTabs/Live/SConcertLiveSessionTabView.cpp b/Engine/Plugins/Developer/Concert/ConcertApp/MultiUserServer/Source/MultiUserServer/Private/Widgets/SessionTabs/Live/SConcertLiveSessionTabView.cpp
--- a/Engine/Plugins/Developer/Concert/ConcertApp/MultiUserServer/Source/MultiUserServer/Private/Widgets/SessionTabs/Live/SConcertLiveSessionTabView.cpp
+++ b/Engine/Plugins/Developer/Concert/ConcertApp/MultiUserServer/Source/MultiUserServer/Private/Widgets/SessionTabs/Live/SConcertLiveSessionTabView.cpp
@@ -47,11 +47,15 @@ void SConcertLiveSessionTabView::Construct(const FArguments& InArgs, const FRequ
 
 TSharedRef<SWidget> SConcertLiveSessionTabView::CreateTabs(const TSharedRef<FTabManager>& InTabManager, const TSharedRef<FTabManager::FLayout>& InLayout, const FRequiredWidgets& InRequiredArgs)
 {
+	const TSharedRef<FWorkspaceItem> WorkspaceItem = InTabManager->AddLocalWorkspaceMenuCategory(LOCTEXT("ArchiveSession", "Live Session"));
+	
 	InTabManager->RegisterTabSpawner(HistoryTabId, FOnSpawnTab::CreateSP(this, &SConcertLiveSessionTabView::SpawnActivityHistory, InRequiredArgs.SessionHistory))
-		.SetDisplayName(LOCTEXT("ActivityHistoryLabel", "History"));
+		.SetDisplayName(LOCTEXT("ActivityHistoryLabel", "History"))
+		.SetGroup(WorkspaceItem);
 	
 	InTabManager->RegisterTabSpawner(SessionContentTabId, FOnSpawnTab::CreateSP(this, &SConcertLiveSessionTabView::SpawnSessionContent, InRequiredArgs.PackageViewer))
-		.SetDisplayName(LOCTEXT("SessionContentLabel", "Session Content"));
+		.SetDisplayName(LOCTEXT("SessionContentLabel", "Session Content"))
+		.SetGroup(WorkspaceItem);
 
 	InLayout->AddArea
 	(
@@ -62,12 +66,13 @@ TSharedRef<SWidget> SConcertLiveSessionTabView::CreateTabs(const TSharedRef<FTab
 				FTabManager::NewStack()
 				->SetSizeCoefficient(0.5f)
 				->AddTab(HistoryTabId, ETabState::OpenedTab)
+				->SetHideTabWell(true)
 			)
 			->Split
 			(
 				FTabManager::NewStack()
 				->SetSizeCoefficient(0.5f)
-				->AddTab(SessionContentTabId, ETabState::OpenedTab)
+				->AddTab(SessionContentTabId, ETabState::ClosedTab)
 			)
 	);
 
diff --git a/Engine/Plugins/Developer/Concert/ConcertApp/MultiUserServer/Source/MultiUserServer/Private/Window/ConcertServerWindowController.cpp b/Engine/Plugins/Developer/Concert/ConcertApp/MultiUserServer/Source/MultiUserServer/Private/Window/ConcertServerWindowController.cpp
--- a/Engine/Plugins/Developer/Concert/ConcertApp/MultiUserServer/Source/MultiUserServer/Private/Window/ConcertServerWindowController.cpp
+++ b/Engine/Plugins/Developer/Concert/ConcertApp/MultiUserServer/Source/MultiUserServer/Private/Window/ConcertServerWindowController.cpp
@@ -20,7 +20,7 @@
 
 #include "Widgets/SWindow.h"
 
-#define LOCTEXT_NAMESPACE "UnrealMultiUserUI"
+#define LOCTEXT_NAMESPACE "UnrealMultiUserUI.FConcertServerWindowController"
 
 namespace UE::MultiUserServer
 {
@@ -37,6 +37,7 @@ namespace UE::MultiUserServer
 
 	FConcertServerWindowController::~FConcertServerWindowController()
 	{
+		FGlobalTabmanager::Get()->SaveAllVisualState();
 		UnregisterFromSessionDestructionEvents();
 	}
 
@@ -46,7 +47,7 @@ namespace UE::MultiUserServer
 		FSlateApplication::Get().GetDisplayMetrics(DisplayMetrics);
 		const float DPIScaleFactor = FPlatformApplicationMisc::GetDPIScaleFactorAtPoint(DisplayMetrics.PrimaryDisplayWorkAreaRect.Left, DisplayMetrics.PrimaryDisplayWorkAreaRect.Top);
 
-		const bool bEmbedTitleAreaContent = false;
+		constexpr bool bEmbedTitleAreaContent = true;
 		const FVector2D ClientSize(1000.0f * DPIScaleFactor, 800.0f * DPIScaleFactor);
 		TSharedRef<SWindow> RootWindowRef = SNew(SWindow)
 			.Title(LOCTEXT("WindowTitle", "Unreal Multi User Server"))
@@ -61,7 +62,7 @@ namespace UE::MultiUserServer
 			.AdjustInitialSizeAndPositionForDPIScale(false);
 		RootWindow = RootWindowRef;
 			
-		const bool bShowRootWindowImmediately = false;
+		constexpr bool bShowRootWindowImmediately = false;
 		FSlateApplication::Get().AddWindow(RootWindowRef, bShowRootWindowImmediately);
 		FGlobalTabmanager::Get()->SetRootWindow(RootWindowRef);
 		FGlobalTabmanager::Get()->SetAllowWindowMenuBar(true);
