Fix menu showing wrong entries on start up

diff --git a/Engine/Plugins/Developer/Concert/ConcertApp/MultiUserServer/Source/MultiUserServer/Private/Widgets/Browser/ConcertServerSessionBrowserController.cpp b/Engine/Plugins/Developer/Concert/ConcertApp/MultiUserServer/Source/MultiUserServer/Private/Widgets/Browser/ConcertServerSessionBrowserController.cpp
--- a/Engine/Plugins/Developer/Concert/ConcertApp/MultiUserServer/Source/MultiUserServer/Private/Widgets/Browser/ConcertServerSessionBrowserController.cpp
+++ b/Engine/Plugins/Developer/Concert/ConcertApp/MultiUserServer/Source/MultiUserServer/Private/Widgets/Browser/ConcertServerSessionBrowserController.cpp
@@ -57,6 +57,11 @@ namespace UE::MultiUserServer
 		ConcertServerEvents::OnArchivedSessionDestroyed().AddSP(this, &FConcertServerSessionBrowserController::OnArchivedSessionDestroyed);
 	}
 
+	void FConcertServerSessionBrowserController::OpenTab()
+	{
+		FGlobalTabmanager::Get()->TryInvokeTab(ConcertServerTabs::GetSessionBrowserTabId());
+	}
+
 	TArray<FConcertServerInfo> FConcertServerSessionBrowserController::GetServers() const
 	{
 		return { ServerInstance->GetConcertServer()->GetServerInfo() };
diff --git a/Engine/Plugins/Developer/Concert/ConcertApp/MultiUserServer/Source/MultiUserServer/Private/Widgets/Browser/ConcertServerSessionBrowserController.h b/Engine/Plugins/Developer/Concert/ConcertApp/MultiUserServer/Source/MultiUserServer/Private/Widgets/Browser/ConcertServerSessionBrowserController.h
--- a/Engine/Plugins/Developer/Concert/ConcertApp/MultiUserServer/Source/MultiUserServer/Private/Widgets/Browser/ConcertServerSessionBrowserController.h
+++ b/Engine/Plugins/Developer/Concert/ConcertApp/MultiUserServer/Source/MultiUserServer/Private/Widgets/Browser/ConcertServerSessionBrowserController.h
@@ -34,6 +34,8 @@ namespace UE::MultiUserServer
 		virtual void Init(const FConcertComponentInitParams& Params) override;
 		//~ End IConcertComponent Interface
 
+		void OpenTab();
+
 		//~ Begin IConcertSessionBrowserController Interface
 		virtual TArray<FConcertServerInfo> GetServers() const override;
 		virtual TArray<FActiveSessionInfo> GetActiveSessions() const override;
diff --git a/Engine/Plugins/Developer/Concert/ConcertApp/MultiUserServer/Source/MultiUserServer/Private/Window/ConcertServerWindowController.cpp b/Engine/Plugins/Developer/Concert/ConcertApp/MultiUserServer/Source/MultiUserServer/Private/Window/ConcertServerWindowController.cpp
--- a/Engine/Plugins/Developer/Concert/ConcertApp/MultiUserServer/Source/MultiUserServer/Private/Window/ConcertServerWindowController.cpp
+++ b/Engine/Plugins/Developer/Concert/ConcertApp/MultiUserServer/Source/MultiUserServer/Private/Window/ConcertServerWindowController.cpp
@@ -85,6 +85,7 @@ namespace UE::MultiUserServer
 		RootWindow->BringToFront(bForceWindowToFront);
 
 		RegisterForSessionDestructionEvents();
+		SessionBrowserController->OpenTab();
 		return RootWindowRef;
 	}
 
