Change pill style

diff --git a/Engine/Plugins/Developer/Concert/ConcertApp/MultiUserServer/Source/MultiUserServer/Private/Widgets/Util/Filter/SConcertFilterBar.h b/Engine/Plugins/Developer/Concert/ConcertApp/MultiUserServer/Source/MultiUserServer/Private/Widgets/Util/Filter/SConcertFilterBar.h
--- a/Engine/Plugins/Developer/Concert/ConcertApp/MultiUserServer/Source/MultiUserServer/Private/Widgets/Util/Filter/SConcertFilterBar.h
+++ b/Engine/Plugins/Developer/Concert/ConcertApp/MultiUserServer/Source/MultiUserServer/Private/Widgets/Util/Filter/SConcertFilterBar.h
@@ -65,7 +65,7 @@ namespace UE::MultiUserServer
 			Args._CreateTextFilter = InArgs._CreateTextFilter;
 			Args._FilterSearchBox = InArgs._FilterSearchBox;
 		
-			SBasicFilterBar<TFilterType>::Construct(Args);
+			SBasicFilterBar<TFilterType>::Construct(Args.FilterPillStyle(EFilterPillStyle::Basic));
 		}
 
 	private:
