Update icons

diff --git a/Engine/Plugins/Developer/Concert/ConcertApp/MultiUserServer/Source/MultiUserServer/Private/ConcertServerStyle.cpp b/Engine/Plugins/Developer/Concert/ConcertApp/MultiUserServer/Source/MultiUserServer/Private/ConcertServerStyle.cpp
--- a/Engine/Plugins/Developer/Concert/ConcertApp/MultiUserServer/Source/MultiUserServer/Private/ConcertServerStyle.cpp
+++ b/Engine/Plugins/Developer/Concert/ConcertApp/MultiUserServer/Source/MultiUserServer/Private/ConcertServerStyle.cpp
@@ -53,12 +53,22 @@ FString FConcertServerStyle::InContent(const FString& RelativePath, const ANSICH
 TSharedRef<FSlateStyleSet> FConcertServerStyle::Create()
 {
 	TSharedRef<FSlateStyleSet> StyleSet = MakeShared<FSlateStyleSet>(GetStyleSetName());
+	StyleSet->SetParentStyleName("CoreStyle");
+	FAppStyle::SetAppStyleSet(*StyleSet); // Makes the app icons appear
 	
 	StyleSet->SetContentRoot(FPaths::EngineContentDir() / TEXT("Slate/Starship/Insights"));
 	StyleSet->SetCoreContentRoot(FPaths::EngineContentDir() / TEXT("Slate"));
 	
 	const FVector2D Icon16x16(16.0f, 16.0f); 
+	const FVector2D Icon24x24(24.0f, 24.0f); 
 	const FVector2D Icon32x32(32.0f, 32.0f); 
+	const FVector2D Icon48x48(48.0f, 48.0f); 
+	
+	// Application icons
+	StyleSet->Set("AppIcon", new IMAGE_PLUGIN_BRUSH("Icons/App/MultiUserApp_48", Icon48x48));
+	StyleSet->Set("AppIconPadding", FMargin(5.0f, 5.0f, 5.0f, 5.0f));
+	StyleSet->Set("AppIcon.Small", new IMAGE_PLUGIN_BRUSH("Icons/App/MultiUserApp_24", Icon24x24));
+	StyleSet->Set("AppIconPadding.Small", FMargin(4.0f, 4.0f, 0.0f, 0.0f));
 	
 	// Icons
 	StyleSet->Set("Concert.MultiUser", new IMAGE_PLUGIN_BRUSH("Icons/icon_MultiUser_32x", Icon32x32));
@@ -80,7 +90,18 @@ TSharedRef<FSlateStyleSet> FConcertServerStyle::Create()
 	StyleSet->Set("Concert.Unmuted",     new IMAGE_PLUGIN_BRUSH("Icons/Unmuted_16x", Icon16x16));
 	// Icons: DependencyView
 	StyleSet->Set("Concert.HighlightHardDependencies",     new IMAGE_PLUGIN_BRUSH("Icons/Ack_Ack_16x", Icon16x16));
-	
+
+	// New icons
+	StyleSet->Set("Concert.Icon.Archive", new IMAGE_PLUGIN_BRUSH_SVG("Icons/Archive_16", Icon16x16));
+	StyleSet->Set("Concert.Icon.Client", new IMAGE_PLUGIN_BRUSH_SVG("Icons/Client_16", Icon16x16));
+	StyleSet->Set("Concert.Icon.CreateMultiUser", new IMAGE_PLUGIN_BRUSH_SVG("Icons/CreateMultiUser_16", Icon16x16));
+	StyleSet->Set("Concert.Icon.Export", new IMAGE_PLUGIN_BRUSH_SVG("Icons/Export_16", Icon16x16));
+	StyleSet->Set("Concert.Icon.Import", new IMAGE_PLUGIN_BRUSH_SVG("Icons/Import_16", Icon16x16));
+	StyleSet->Set("Concert.Icon.LogServer", new IMAGE_PLUGIN_BRUSH_SVG("Icons/LogServer", Icon16x16));
+	StyleSet->Set("Concert.Icon.LogSession", new IMAGE_PLUGIN_BRUSH_SVG("Icons/LogSession", Icon16x16));
+	StyleSet->Set("Concert.Icon.MultiUser", new IMAGE_PLUGIN_BRUSH_SVG("Icons/MultiUser_16", Icon16x16));
+	StyleSet->Set("Concert.Icon.Package", new IMAGE_PLUGIN_BRUSH_SVG("Icons/Package", Icon16x16));
+	StyleSet->Set("Concert.Icon.Server", new IMAGE_PLUGIN_BRUSH_SVG("Icons/Server_16", Icon16x16));
 
 	// Clients tab
 	StyleSet->Set("Concert.Clients.DropShadow", new IMAGE_PLUGIN_BOX_BRUSH("ClientThumbnailDropShadow", FMargin(4.0f / 64.0f)));
@@ -90,18 +111,18 @@ TSharedRef<FSlateStyleSet> FConcertServerStyle::Create()
 	StyleSet->Set("Concert.Clients.ThumbnailCurveBackground", new FSlateRoundedBoxBrush(FStyleColors::Secondary, 0.0f));
 	StyleSet->Set("Concert.Clients.ThumbnailFooter", new FSlateRoundedBoxBrush(FStyleColors::Panel, 0.0f));
 	StyleSet->Set("Concert.Clients.TileTableRow", FTableRowStyle()
-			.SetEvenRowBackgroundBrush(FSlateNoResource() )
-			.SetEvenRowBackgroundHoveredBrush(FSlateNoResource())
-			.SetOddRowBackgroundBrush(FSlateNoResource())
-			.SetOddRowBackgroundHoveredBrush(FSlateNoResource())
-			.SetSelectorFocusedBrush(FSlateNoResource())
-			.SetActiveBrush(FSlateNoResource())
-			.SetActiveHoveredBrush(FSlateNoResource())
-			.SetInactiveBrush(FSlateNoResource())
-			.SetInactiveHoveredBrush(FSlateNoResource())
-			.SetTextColor(FSlateColor())
-			.SetSelectedTextColor(FSlateColor())
-			);
+		.SetEvenRowBackgroundBrush(FSlateNoResource() )
+		.SetEvenRowBackgroundHoveredBrush(FSlateNoResource())
+		.SetOddRowBackgroundBrush(FSlateNoResource())
+		.SetOddRowBackgroundHoveredBrush(FSlateNoResource())
+		.SetSelectorFocusedBrush(FSlateNoResource())
+		.SetActiveBrush(FSlateNoResource())
+		.SetActiveHoveredBrush(FSlateNoResource())
+		.SetInactiveBrush(FSlateNoResource())
+		.SetInactiveHoveredBrush(FSlateNoResource())
+		.SetTextColor(FSlateColor())
+		.SetSelectedTextColor(FSlateColor())
+		);
 	StyleSet->Set("Concert.Clients.ClientNameTileFont", DEFAULT_FONT("Regular", 16));
 	
 	// Graph colours
diff --git a/Engine/Plugins/Developer/Concert/ConcertApp/MultiUserServer/Source/MultiUserServer/Private/Widgets/Browser/ConcertServerSessionBrowserController.cpp b/Engine/Plugins/Developer/Concert/ConcertApp/MultiUserServer/Source/MultiUserServer/Private/Widgets/Browser/ConcertServerSessionBrowserController.cpp
--- a/Engine/Plugins/Developer/Concert/ConcertApp/MultiUserServer/Source/MultiUserServer/Private/Widgets/Browser/ConcertServerSessionBrowserController.cpp
+++ b/Engine/Plugins/Developer/Concert/ConcertApp/MultiUserServer/Source/MultiUserServer/Private/Widgets/Browser/ConcertServerSessionBrowserController.cpp
@@ -46,7 +46,7 @@ namespace UE::MultiUserServer
 			)
 			.SetDisplayName(LOCTEXT("SessionBrowserTabTitleLong", "Session Browser"))
 			.SetTooltipText(LOCTEXT("SessionBrowserTooltipText", "A section to browse, start, archive, and restore server sessions."))
-			.SetIcon(FSlateIcon(FConcertServerStyle::GetStyleSetName(), TEXT("Concert.MultiUser"))
+			.SetIcon(FSlateIcon(FConcertServerStyle::GetStyleSetName(), TEXT("Concert.Icon.MultiUser"))
 		);
 		Params.MainStack->AddTab(ConcertServerTabs::GetSessionBrowserTabId(), ETabState::OpenedTab)
 			->SetForegroundTab(ConcertServerTabs::GetSessionBrowserTabId());
diff --git a/Engine/Plugins/Developer/Concert/ConcertApp/MultiUserServer/Source/MultiUserServer/Private/Widgets/Browser/SConcertServerSessionBrowser.cpp b/Engine/Plugins/Developer/Concert/ConcertApp/MultiUserServer/Source/MultiUserServer/Private/Widgets/Browser/SConcertServerSessionBrowser.cpp
--- a/Engine/Plugins/Developer/Concert/ConcertApp/MultiUserServer/Source/MultiUserServer/Private/Widgets/Browser/SConcertServerSessionBrowser.cpp
+++ b/Engine/Plugins/Developer/Concert/ConcertApp/MultiUserServer/Source/MultiUserServer/Private/Widgets/Browser/SConcertServerSessionBrowser.cpp
@@ -2,6 +2,7 @@
 
 #include "SConcertServerSessionBrowser.h"
 
+#include "ConcertServerStyle.h"
 #include "MultiUserServerModule.h"
 #include "Session/Browser/ConcertBrowserUtils.h"
 #include "Session/Browser/Items/ConcertSessionTreeItem.h"
@@ -42,7 +43,8 @@ namespace UE::MultiUserServer
 		
 		InTabManager->RegisterTabSpawner(SessionBrowserTabId, FOnSpawnTab::CreateSP(this, &SConcertServerSessionBrowser::SpawnSessionBrowserTab, InArgs._DoubleClickLiveSession, InArgs._DoubleClickArchivedSession))
 			.SetDisplayName(LOCTEXT("SessionBrowserTabLabel", "Sessions"))
-			.SetGroup(WorkspaceItem);
+			.SetGroup(WorkspaceItem)
+			.SetIcon(FSlateIcon(FConcertServerStyle::GetStyleSetName(), TEXT("Concert.Icon.MultiUser")));
 
 		InLayout->AddArea
 		(
diff --git a/Engine/Plugins/Developer/Concert/ConcertApp/MultiUserServer/Source/MultiUserServer/Private/Widgets/Clients/ConcertClientsTabController.cpp b/Engine/Plugins/Developer/Concert/ConcertApp/MultiUserServer/Source/MultiUserServer/Private/Widgets/Clients/ConcertClientsTabController.cpp
--- a/Engine/Plugins/Developer/Concert/ConcertApp/MultiUserServer/Source/MultiUserServer/Private/Widgets/Clients/ConcertClientsTabController.cpp
+++ b/Engine/Plugins/Developer/Concert/ConcertApp/MultiUserServer/Source/MultiUserServer/Private/Widgets/Clients/ConcertClientsTabController.cpp
@@ -2,6 +2,7 @@
 
 #include "ConcertClientsTabController.h"
 
+#include "ConcertServerStyle.h"
 #include "IConcertSyncServer.h"
 #include "SConcertClientsTabView.h"
 #include "Logging/Source/GlobalLogSource.h"
@@ -26,7 +27,8 @@ namespace UE::MultiUserServer
 					FOnSpawnTab::CreateRaw(this, &FConcertClientsTabController::SpawnClientsTab, Params.WindowController->GetRootWindow(), Params.Server)
 				)
 				.SetDisplayName(LOCTEXT("ClientsTabTitle", "Clients"))
-				.SetTooltipText(LOCTEXT("ClientsTooltipText", "View network statistics for connected clients.")
+				.SetTooltipText(LOCTEXT("ClientsTooltipText", "View network statistics for connected clients."))
+				.SetIcon(FSlateIcon(FConcertServerStyle::GetStyleSetName(), TEXT("Concert.Icon.Client"))
 			);
 		Params.MainStack->AddTab(ConcertServerTabs::GetClientsTabID(), ETabState::OpenedTab);
 	}
diff --git a/Engine/Plugins/Developer/Concert/ConcertApp/MultiUserServer/Source/MultiUserServer/Private/Widgets/Clients/PackageTransmission/PackageTransmissionTabController.cpp b/Engine/Plugins/Developer/Concert/ConcertApp/MultiUserServer/Source/MultiUserServer/Private/Widgets/Clients/PackageTransmission/PackageTransmissionTabController.cpp
--- a/Engine/Plugins/Developer/Concert/ConcertApp/MultiUserServer/Source/MultiUserServer/Private/Widgets/Clients/PackageTransmission/PackageTransmissionTabController.cpp
+++ b/Engine/Plugins/Developer/Concert/ConcertApp/MultiUserServer/Source/MultiUserServer/Private/Widgets/Clients/PackageTransmission/PackageTransmissionTabController.cpp
@@ -2,6 +2,7 @@
 
 #include "PackageTransmissionTabController.h"
 
+#include "ConcertServerStyle.h"
 #include "Model/PackageTransmissionModel.h"
 #include "SPackageTransmissionView.h"
 #include "Util/PackageTransmissionEntryTokenizer.h"
@@ -33,7 +34,8 @@ namespace UE::MultiUserServer
 		OwningTabManager->RegisterTabSpawner(TabId, FOnSpawnTab::CreateRaw(this, &FPackageTransmissionTabController::SpawnTab))
 			// In the future we may create multiple FPackageTransmissionTabController and may want to change this name to depend on some parameter
 			.SetDisplayName(LOCTEXT("PackageTabLabel", "Packages"))
-			.SetGroup(WorkspaceItem);
+			.SetGroup(WorkspaceItem)
+			.SetIcon(FSlateIcon(FConcertServerStyle::GetStyleSetName(), TEXT("Concert.Icon.Package")));
 	}
 
 	FPackageTransmissionTabController::~FPackageTransmissionTabController()
diff --git a/Engine/Plugins/Developer/Concert/ConcertApp/MultiUserServer/Source/MultiUserServer/Private/Widgets/Clients/SConcertClientsTabView.cpp b/Engine/Plugins/Developer/Concert/ConcertApp/MultiUserServer/Source/MultiUserServer/Private/Widgets/Clients/SConcertClientsTabView.cpp
--- a/Engine/Plugins/Developer/Concert/ConcertApp/MultiUserServer/Source/MultiUserServer/Private/Widgets/Clients/SConcertClientsTabView.cpp
+++ b/Engine/Plugins/Developer/Concert/ConcertApp/MultiUserServer/Source/MultiUserServer/Private/Widgets/Clients/SConcertClientsTabView.cpp
@@ -2,6 +2,7 @@
 
 #include "SConcertClientsTabView.h"
 
+#include "ConcertServerStyle.h"
 #include "IConcertSyncServer.h"
 #include "Logging/Filter/ConcertLogFilter_FrontendRoot.h"
 #include "Logging/Source/GlobalLogSource.h"
@@ -77,7 +78,9 @@ namespace UE::MultiUserServer
 				[
 					SNew(SConcertTransportLog, LogBuffer.ToSharedRef(), ClientInfoCache.ToSharedRef(), LogTokenizer.ToSharedRef())
 					.Filter(UE::MultiUserServer::MakeClientLogFilter(LogTokenizer.ToSharedRef(), ClientMessageNodeId, ClientInfoCache.ToSharedRef()))
-				]; 
+				];
+
+			NewTab->SetTabIcon(FConcertServerStyle::Get().GetBrush(TEXT("Concert.Icon.LogSession")));
 
 			// We need a tab to place the client tab next to
 			if (IsGlobalLogOpen())
@@ -153,10 +156,12 @@ namespace UE::MultiUserServer
 			);
 		InTabManager->RegisterTabSpawner(ClientBrowserTabId, FOnSpawnTab::CreateSP(this, &SConcertClientsTabView::SpawnClientBrowserTab))
 			.SetDisplayName(LOCTEXT("ClientBrowserTabLabel", "Clients"))
-			.SetGroup(WorkspaceItem);
+			.SetGroup(WorkspaceItem)
+			.SetIcon(FSlateIcon(FConcertServerStyle::GetStyleSetName(), TEXT("Concert.Icon.Client")));
 		InTabManager->RegisterTabSpawner(GlobalLogTabId, FOnSpawnTab::CreateSP(this, &SConcertClientsTabView::SpawnGlobalLogTab))
 			.SetDisplayName(LOCTEXT("ServerLogTabLabel", "Server Log"))
-			.SetGroup(WorkspaceItem);
+			.SetGroup(WorkspaceItem)
+			.SetIcon(FSlateIcon(FConcertServerStyle::GetStyleSetName(), TEXT("Concert.Icon.LogServer")));
 		
 		InLayout->AddArea
 			(
@@ -247,7 +252,7 @@ namespace UE::MultiUserServer
 					[
 						SNew(SImage)
 						.ColorAndOpacity(FSlateColor::UseForeground())
-						.Image(FAppStyle::Get().GetBrush("Icons.Layout"))
+						.Image(FConcertServerStyle::Get().GetBrush(TEXT("Concert.Icon.LogServer")))
 					]
 					+SHorizontalBox::Slot()
 					.VAlign(VAlign_Center)
diff --git a/Engine/Plugins/Developer/Concert/ConcertApp/MultiUserServer/Source/MultiUserServer/Private/Widgets/SConcertTabViewWithManagerBase.cpp b/Engine/Plugins/Developer/Concert/ConcertApp/MultiUserServer/Source/MultiUserServer/Private/Widgets/SConcertTabViewWithManagerBase.cpp
--- a/Engine/Plugins/Developer/Concert/ConcertApp/MultiUserServer/Source/MultiUserServer/Private/Widgets/SConcertTabViewWithManagerBase.cpp
+++ b/Engine/Plugins/Developer/Concert/ConcertApp/MultiUserServer/Source/MultiUserServer/Private/Widgets/SConcertTabViewWithManagerBase.cpp
@@ -2,9 +2,12 @@
 
 #include "SConcertTabViewWithManagerBase.h"
 
-#include "Framework/MultiBox/MultiBoxBuilder.h"
+#include "Widgets/Util/SMultiUserIcons.h"
+
+#include "Framework/Application/SlateApplication.h"
 #include "Framework/Docking/LayoutService.h"
 #include "Framework/Docking/TabManager.h"
+#include "Framework/MultiBox/MultiBoxBuilder.h"
 
 #define LOCTEXT_NAMESPACE "UnrealMultiUserUI.SConcertTabViewWithManagerBase"
 
@@ -61,6 +64,12 @@ void SConcertTabViewWithManagerBase::FillInDefaultMenuItems(FMenuBarBuilder Menu
 			FNewMenuDelegate::CreateSP(this, &SConcertTabViewWithManagerBase::FillWindowMenu),
 			"Window"
 		);
+	MenuBarBuilder.AddPullDownMenu(
+			LOCTEXT("DebugMenuLabel", "Debug"),
+			FText::GetEmpty(),
+			FNewMenuDelegate::CreateSP(this, &SConcertTabViewWithManagerBase::FillDebugMenu),
+			"Debug"
+		);
 }
 
 void SConcertTabViewWithManagerBase::FillWindowMenu(FMenuBuilder& MenuBuilder)
@@ -68,4 +77,32 @@ void SConcertTabViewWithManagerBase::FillWindowMenu(FMenuBuilder& MenuBuilder)
 	TabManager->PopulateLocalTabSpawnerMenu(MenuBuilder);
 }
 
+void SConcertTabViewWithManagerBase::FillDebugMenu(FMenuBuilder& MenuBuilder)
+{
+	MenuBuilder.AddMenuEntry(
+		LOCTEXT("Icons", "View App Icons"),
+		FText::GetEmpty(),
+		FSlateIcon(),
+		FUIAction(
+			FExecuteAction::CreateRaw(this, &SConcertTabViewWithManagerBase::AddIconWindow)
+			)
+		);
+}
+
+void SConcertTabViewWithManagerBase::AddIconWindow()
+{
+	const FText AboutWindowTitle = LOCTEXT("UnrealMultiUserIcons", "Unreal Multi-User Server Icons");
+	TSharedPtr<SWindow> AboutWindow = 
+		SNew(SWindow)
+		.Title( AboutWindowTitle )
+		.ClientSize(FVector2D(720.f, 538.f))
+		.SupportsMaximize(false) .SupportsMinimize(false)
+		.SizingRule( ESizingRule::FixedSize )
+		[
+			SNew(SMultiUserIcons)
+		];
+
+	FSlateApplication::Get().AddWindow(AboutWindow.ToSharedRef());
+}
+
 #undef LOCTEXT_NAMESPACE
diff --git a/Engine/Plugins/Developer/Concert/ConcertApp/MultiUserServer/Source/MultiUserServer/Private/Widgets/SConcertTabViewWithManagerBase.h b/Engine/Plugins/Developer/Concert/ConcertApp/MultiUserServer/Source/MultiUserServer/Private/Widgets/SConcertTabViewWithManagerBase.h
--- a/Engine/Plugins/Developer/Concert/ConcertApp/MultiUserServer/Source/MultiUserServer/Private/Widgets/SConcertTabViewWithManagerBase.h
+++ b/Engine/Plugins/Developer/Concert/ConcertApp/MultiUserServer/Source/MultiUserServer/Private/Widgets/SConcertTabViewWithManagerBase.h
@@ -51,4 +51,7 @@ private:
 	TSharedRef<SWidget> CreateTabs(const FArguments& InArgs);
 	void FillInDefaultMenuItems(FMenuBarBuilder MenuBarBuilder);
 	void FillWindowMenu(FMenuBuilder& MenuBuilder);
+	void FillDebugMenu(FMenuBuilder& MenuBuilder);
+
+	void AddIconWindow();
 };
diff --git a/Engine/Plugins/Developer/Concert/ConcertApp/MultiUserServer/Source/MultiUserServer/Private/Widgets/SessionTabs/Live/SConcertLiveSessionTabView.cpp b/Engine/Plugins/Developer/Concert/ConcertApp/MultiUserServer/Source/MultiUserServer/Private/Widgets/SessionTabs/Live/SConcertLiveSessionTabView.cpp
--- a/Engine/Plugins/Developer/Concert/ConcertApp/MultiUserServer/Source/MultiUserServer/Private/Widgets/SessionTabs/Live/SConcertLiveSessionTabView.cpp
+++ b/Engine/Plugins/Developer/Concert/ConcertApp/MultiUserServer/Source/MultiUserServer/Private/Widgets/SessionTabs/Live/SConcertLiveSessionTabView.cpp
@@ -2,6 +2,7 @@
 
 #include "SConcertLiveSessionTabView.h"
 
+#include "ConcertServerStyle.h"
 #include "Framework/Docking/TabManager.h"
 #include "Widgets/Docking/SDockTab.h"
 #include "Session/History/SSessionHistory.h"
@@ -55,7 +56,8 @@ TSharedRef<SWidget> SConcertLiveSessionTabView::CreateTabs(const TSharedRef<FTab
 	
 	InTabManager->RegisterTabSpawner(SessionContentTabId, FOnSpawnTab::CreateSP(this, &SConcertLiveSessionTabView::SpawnSessionContent, InRequiredArgs.PackageViewer))
 		.SetDisplayName(LOCTEXT("SessionContentLabel", "Session Content"))
-		.SetGroup(WorkspaceItem);
+		.SetGroup(WorkspaceItem)
+		.SetIcon(FSlateIcon(FConcertServerStyle::GetStyleSetName(), TEXT("Concert.Icon.Package")));
 
 	InLayout->AddArea
 	(
@@ -124,7 +126,7 @@ TSharedRef<SWidget> SConcertLiveSessionTabView::CreateToolbar(const FArguments&
 				[
 					SNew(SImage)
 					.ColorAndOpacity(FSlateColor::UseForeground())
-					.Image(FAppStyle::Get().GetBrush("Icons.Layout"))
+					.Image(FConcertServerStyle::Get().GetBrush("Concert.Icon.Client"))
 				]
 				+SHorizontalBox::Slot()
 				.VAlign(VAlign_Center)
diff --git a/Engine/Plugins/Developer/Concert/ConcertApp/MultiUserServer/Source/MultiUserServer/Private/Widgets/Util/SMultiUserIcons.cpp b/Engine/Plugins/Developer/Concert/ConcertApp/MultiUserServer/Source/MultiUserServer/Private/Widgets/Util/SMultiUserIcons.cpp
new file mode 100644
--- /dev/null
+++ b/Engine/Plugins/Developer/Concert/ConcertApp/MultiUserServer/Source/MultiUserServer/Private/Widgets/Util/SMultiUserIcons.cpp
@@ -0,0 +1,109 @@
+// Copyright Epic Games, Inc. All Rights Reserved.
+
+#include "SMultiUserIcons.h"
+
+#include "ConcertServerStyle.h"
+
+#include "Widgets/Images/SImage.h"
+#include "Widgets/Layout/SBorder.h"
+#include "Widgets/SBoxPanel.h"
+#include "Widgets/Text/STextBlock.h"
+
+#define LOCTEXT_NAMESPACE "UnrealMultiUserUI.SMultiUserIcons"
+
+void SMultiUserIcons::Construct(const FArguments& InArgs)
+{
+	const TSharedRef<SVerticalBox> VerticalBox = SNew(SVerticalBox);
+
+	// Please keep this alphabetically sorted if you add new entries
+	
+	// App
+	AddHeader(VerticalBox, LOCTEXT("App", "App"));
+	AddCategory(VerticalBox, LOCTEXT("App", "App"), { { TEXT("AppIcon"), TEXT("AppIcon.Small") } });
+	
+	// New
+	AddHeader(VerticalBox, LOCTEXT("Icons", "New Icons"));
+	AddCategory(VerticalBox, LOCTEXT("Ack", "Ack"), { TEXT("Concert.Ack.Ack"), TEXT("Concert.Ack.Success"), TEXT("Concert.Ack.Failure") });
+	AddCategory(VerticalBox, LOCTEXT("Archive", "Archive"), { TEXT("Concert.Icon.Archive") });
+	AddCategory(VerticalBox, LOCTEXT("Client", "Client"), { TEXT("Concert.Icon.Client") });
+	AddCategory(VerticalBox, LOCTEXT("CreateMultiUser", "CreateMultiUser"), { TEXT("Concert.Icon.CreateMultiUser") });
+	AddCategory(VerticalBox, LOCTEXT("ExportImport", "ExportImport"), { TEXT("Concert.Icon.Export"), TEXT("Concert.Icon.Import") });
+	AddCategory(VerticalBox, LOCTEXT("LogServer", "LogServer"), { TEXT("Concert.Icon.LogServer") });
+	AddCategory(VerticalBox, LOCTEXT("LogSession", "LogSession"), { TEXT("Concert.Icon.LogSession") });
+	AddCategory(VerticalBox, LOCTEXT("MultiUser", "MultiUser"), { TEXT("Concert.Icon.MultiUser") });
+	AddCategory(VerticalBox, LOCTEXT("Muting", "Muting"), { TEXT("Concert.Muted"), TEXT("Concert.Unmuted") });
+	AddCategory(VerticalBox, LOCTEXT("Package", "Package"), { TEXT("Concert.Icon.Package") });
+	AddCategory(VerticalBox, LOCTEXT("PackageChanged", "Package Changed"), { TEXT("Concert.SessionContent.PackageAdded"), TEXT("Concert.SessionContent.PackageDeleted"), TEXT("Concert.SessionContent.PackageRenamed"), TEXT("Concert.SessionContent.PackageSaved") });
+	AddCategory(VerticalBox, LOCTEXT("PackageTransmission", "PackageTransmission"), { TEXT("Concert.PackageTransmission.Success"), TEXT("Concert.PackageTransmission.Failure") });
+	AddCategory(VerticalBox, LOCTEXT("Server", "Server"), { TEXT("Concert.Icon.Server") });
+	
+	ChildSlot
+	[
+		SNew(SBorder)
+		.Padding(20.f)
+		[
+			VerticalBox
+		]
+	];
+}
+
+void SMultiUserIcons::AddHeader(const TSharedRef<SVerticalBox>& VerticalBox, FText Text)
+{
+	VerticalBox->AddSlot()
+	.AutoHeight()
+	.Padding(0.f, 10.f, 0.f, 5.f)
+	[
+		SNew(SHorizontalBox)
+		+SHorizontalBox::Slot()
+		.Padding(0.f, 0.f, 20.f, 0.f)
+		.HAlign(HAlign_Center)
+		.VAlign(VAlign_Center)
+		[
+			SNew(STextBlock)
+			.Text(Text)
+			.Justification(ETextJustify::Right)
+		]
+	];
+}
+
+void SMultiUserIcons::AddCategory(const TSharedRef<SVerticalBox>& VerticalBox, FText Text, const TArray<FName>& IconNames)
+{
+	TSharedRef<SHorizontalBox> HorizontalBox = SNew(SHorizontalBox);
+	
+	VerticalBox->AddSlot()
+	.AutoHeight()
+	[
+		SNew(SHorizontalBox)
+		+SHorizontalBox::Slot()
+		.Padding(0.f, 0.f, 20.f, 0.f)
+		.FillWidth(0.3f)
+		.HAlign(HAlign_Right)
+		.VAlign(VAlign_Center)
+		[
+			SNew(STextBlock)
+			.Text(Text)
+			.Justification(ETextJustify::Right)
+		]
+		
+		+SHorizontalBox::Slot()
+		.FillWidth(0.7f)
+		[
+			HorizontalBox
+		]
+	];
+	
+	for (const FName& IconName : IconNames)
+	{
+		HorizontalBox->AddSlot()
+		.Padding(5.f)
+		.AutoWidth()
+		.HAlign(HAlign_Center)
+		.VAlign(VAlign_Center)
+		[
+			SNew(SImage)
+			.Image(FConcertServerStyle::Get().GetBrush(IconName))
+		];
+	}
+}
+
+#undef LOCTEXT_NAMESPACE
\ No newline at end of file
diff --git a/Engine/Plugins/Developer/Concert/ConcertApp/MultiUserServer/Source/MultiUserServer/Private/Widgets/Util/SMultiUserIcons.h b/Engine/Plugins/Developer/Concert/ConcertApp/MultiUserServer/Source/MultiUserServer/Private/Widgets/Util/SMultiUserIcons.h
new file mode 100644
--- /dev/null
+++ b/Engine/Plugins/Developer/Concert/ConcertApp/MultiUserServer/Source/MultiUserServer/Private/Widgets/Util/SMultiUserIcons.h
@@ -0,0 +1,26 @@
+// Copyright Epic Games, Inc. All Rights Reserved.
+
+#pragma once
+
+#include "CoreMinimal.h"
+#include "Widgets/SCompoundWidget.h"
+#include "Widgets/DeclarativeSyntaxSupport.h"
+
+class SVerticalBox;
+
+/** Shows all multi user server icons */
+class SMultiUserIcons : public SCompoundWidget
+{
+public:
+
+	SLATE_BEGIN_ARGS(SMultiUserIcons)
+	{}
+	SLATE_END_ARGS()
+
+	void Construct(const FArguments& InArgs);
+
+private:
+
+	void AddHeader(const TSharedRef<SVerticalBox>& VerticalBox, FText Text);
+	void AddCategory(const TSharedRef<SVerticalBox>& VerticalBox, FText Text, const TArray<FName>& IconNames);
+};
diff --git a/Engine/Plugins/Developer/Concert/ConcertUI/ConcertSharedSlate/Source/ConcertSharedSlate/Private/ConcertFrontendStyle.cpp b/Engine/Plugins/Developer/Concert/ConcertUI/ConcertSharedSlate/Source/ConcertSharedSlate/Private/ConcertFrontendStyle.cpp
--- a/Engine/Plugins/Developer/Concert/ConcertUI/ConcertSharedSlate/Source/ConcertSharedSlate/Private/ConcertFrontendStyle.cpp
+++ b/Engine/Plugins/Developer/Concert/ConcertUI/ConcertSharedSlate/Source/ConcertSharedSlate/Private/ConcertFrontendStyle.cpp
@@ -78,7 +78,7 @@ void FConcertFrontendStyle::Initialize()
 	StyleSet->Set("Concert.OpenSettings", new IMAGE_PLUGIN_BRUSH("Icons/icon_Settings_32x",  Icon16x16, IconColorAndOpacity));
 	StyleSet->Set("Concert.LaunchServer", new IMAGE_PLUGIN_BRUSH("Icons/icon_NewServer_32x", Icon16x16, IconColorAndOpacity));
 
-	// Multi-User Browser
+	// Multi-User Browser 
 	StyleSet->Set("Concert.ArchiveSession",			new IMAGE_PLUGIN_BRUSH("Icons/icon_ArchiveSession_48x",			Icon24x24, IconColorAndOpacity));
 	StyleSet->Set("Concert.CancelAutoJoin",			new IMAGE_PLUGIN_BRUSH("Icons/icon_CancelAutoJoin_48x",			Icon24x24, IconColorAndOpacity));
 	StyleSet->Set("Concert.CloseServer",				new IMAGE_PLUGIN_BRUSH("Icons/icon_CloseServer_48x",			Icon24x24, IconColorAndOpacity));
@@ -87,15 +87,15 @@ void FConcertFrontendStyle::Initialize()
 	StyleSet->Set("Concert.JoinSession",				new IMAGE_PLUGIN_BRUSH("Icons/icon_JoinSelectedSession_48x",	Icon24x24, IconColorAndOpacity));
 	StyleSet->Set("Concert.LeaveSession",			new IMAGE_PLUGIN_BRUSH("Icons/icon_LeaveSession_48x",			Icon24x24, IconColorAndOpacity));
 	StyleSet->Set("Concert.NewServer",				new IMAGE_PLUGIN_BRUSH("Icons/icon_NewServer_48x",				Icon24x24, IconColorAndOpacity));
-	StyleSet->Set("Concert.NewSession",				new IMAGE_PLUGIN_BRUSH("Icons/icon_NewSession_48x",				Icon24x24, IconColorAndOpacity));
+	StyleSet->Set("Concert.NewSession",				new IMAGE_PLUGIN_BRUSH_SVG("Icons/icon_CreateMultiUser_16",		Icon24x24, IconColorAndOpacity));
 	StyleSet->Set("Concert.PauseSession",			new IMAGE_PLUGIN_BRUSH("Icons/icon_PauseSession_48x",			Icon24x24, IconColorAndOpacity));
 	StyleSet->Set("Concert.RestoreSession",			new IMAGE_PLUGIN_BRUSH("Icons/icon_RestoreSession_48x",			Icon24x24, IconColorAndOpacity));
 	StyleSet->Set("Concert.ResumeSession",			new IMAGE_PLUGIN_BRUSH("Icons/icon_ResumeSession_48x",			Icon24x24, IconColorAndOpacity));
 	StyleSet->Set("Concert.Settings",				new IMAGE_PLUGIN_BRUSH("Icons/icon_Settings_48x",				Icon24x24, IconColorAndOpacity));
 	StyleSet->Set("Concert.NewServer.Small",			new IMAGE_PLUGIN_BRUSH("Icons/icon_NewServer_32x",				Icon16x16, IconColorAndOpacity));
-	StyleSet->Set("Concert.NewSession.Small",		new IMAGE_PLUGIN_BRUSH("Icons/icon_NewSession_32x",				Icon16x16, IconColorAndOpacity));
-	StyleSet->Set("Concert.ActiveSession.Icon",		new IMAGE_PLUGIN_BRUSH_SVG("Icons/icon_MultiUser",				Icon16x16, ActiveIconColorAndOpacity));
-	StyleSet->Set("Concert.ArchivedSession.Icon",	new IMAGE_PLUGIN_BRUSH_SVG("Icons/icon_MultiUser",				Icon16x16, ArchivedIconColorAndOpacity));
+	StyleSet->Set("Concert.NewSession.Small",		new IMAGE_PLUGIN_BRUSH_SVG("Icons/icon_CreateMultiUser_16",		Icon16x16, IconColorAndOpacity));
+	StyleSet->Set("Concert.ActiveSession.Icon",		new IMAGE_PLUGIN_BRUSH_SVG("Icons/icon_SessionActive_16",		Icon16x16, ActiveIconColorAndOpacity));
+	StyleSet->Set("Concert.ArchivedSession.Icon",	new IMAGE_PLUGIN_BRUSH_SVG("Icons/icon_SessionArchived_16",		Icon16x16, ArchivedIconColorAndOpacity));
 	StyleSet->Set("Concert.SessionBrowser.FontSize", 11.f);
 	StyleSet->Set("Concert.SessionRowPadding", FMargin(0.f, 4.f));
 	
diff --git a/Engine/Plugins/Developer/Concert/ConcertUI/ConcertSharedSlate/Source/ConcertSharedSlate/Private/Session/Browser/SConcertSessionBrowser.cpp b/Engine/Plugins/Developer/Concert/ConcertUI/ConcertSharedSlate/Source/ConcertSharedSlate/Private/Session/Browser/SConcertSessionBrowser.cpp
--- a/Engine/Plugins/Developer/Concert/ConcertUI/ConcertSharedSlate/Source/ConcertSharedSlate/Private/Session/Browser/SConcertSessionBrowser.cpp
+++ b/Engine/Plugins/Developer/Concert/ConcertUI/ConcertSharedSlate/Source/ConcertSharedSlate/Private/Session/Browser/SConcertSessionBrowser.cpp
@@ -471,8 +471,7 @@ TSharedRef<SWidget> SConcertSessionBrowser::MakeButtonBar(const FArguments& InAr
 			FConcertFrontendStyle::Get()->GetBrush("Concert.RestoreSession"),
 			LOCTEXT("RestoreButtonTooltip", "Restore the selected session"),
 			TAttribute<bool>(this, &SConcertSessionBrowser::IsRestoreButtonEnabledInternal),
-			FOnClicked::CreateSP(this, &SConcertSessionBrowser::OnRestoreButtonClicked),
-			TAttribute<EVisibility>::Create([this]() { return IsRestoreButtonEnabledInternal() ? EVisibility::Visible : EVisibility::Collapsed; })
+			FOnClicked::CreateSP(this, &SConcertSessionBrowser::OnRestoreButtonClicked)
 			)
 	);
 	// Archive.
