Fix Multi-User - Slate server crash after modifying Media Profile

diff --git a/Engine/Plugins/Developer/Concert/ConcertApp/MultiUserServer/Source/MultiUserServer/Private/Widgets/Clients/SConcertClientsTabView.cpp b/Engine/Plugins/Developer/Concert/ConcertApp/MultiUserServer/Source/MultiUserServer/Private/Widgets/Clients/SConcertClientsTabView.cpp
--- a/Engine/Plugins/Developer/Concert/ConcertApp/MultiUserServer/Source/MultiUserServer/Private/Widgets/Clients/SConcertClientsTabView.cpp
+++ b/Engine/Plugins/Developer/Concert/ConcertApp/MultiUserServer/Source/MultiUserServer/Private/Widgets/Clients/SConcertClientsTabView.cpp
@@ -102,6 +102,12 @@ namespace UE::MultiUserServer
 
 	bool SConcertClientsTabView::CanScrollToLog(const FGuid& MessageId, FConcertLogEntryFilterFunc FilterFunc, FText& ErrorMessage) const
 	{
+		// GlobalTransportLog can be nullptr if the tab has not (yet) been created
+		if (!GlobalTransportLog)
+		{
+			return false;
+		}
+		
 		const bool bCanScrollToLog = GlobalTransportLog->CanScrollToLog(MessageId, FilterFunc);
 		if (!bCanScrollToLog)
 		{
@@ -115,7 +121,11 @@ namespace UE::MultiUserServer
 	void SConcertClientsTabView::ScrollToLog(const FGuid& MessageId, FConcertLogEntryFilterFunc FilterFunc) const
 	{
 		OpenGlobalLogTab();
-		GlobalTransportLog->ScrollToLog(MessageId, FilterFunc);
+		
+		if (ensureMsgf(GlobalTransportLog, TEXT("OpenGlobalLogTab failed to create the global log widget. Investigate.")))
+		{
+			GlobalTransportLog->ScrollToLog(MessageId, FilterFunc);
+		}
 	}
 
 	void SConcertClientsTabView::OpenGlobalLogTab() const
