Fix crash due to race condition

diff --git a/Engine/Plugins/Developer/Concert/ConcertApp/MultiUserServer/Source/MultiUserServer/Private/Widgets/Clients/Browser/Models/Transfer/ClientTransferStatisticsModel.cpp b/Engine/Plugins/Developer/Concert/ConcertApp/MultiUserServer/Source/MultiUserServer/Private/Widgets/Clients/Browser/Models/Transfer/ClientTransferStatisticsModel.cpp
--- a/Engine/Plugins/Developer/Concert/ConcertApp/MultiUserServer/Source/MultiUserServer/Private/Widgets/Clients/Browser/Models/Transfer/ClientTransferStatisticsModel.cpp
+++ b/Engine/Plugins/Developer/Concert/ConcertApp/MultiUserServer/Source/MultiUserServer/Private/Widgets/Clients/Browser/Models/Transfer/ClientTransferStatisticsModel.cpp
@@ -39,12 +39,13 @@ namespace UE::MultiUserServer
 	}
 	
 	FClientTransferStatisticsModel::FClientTransferStatisticsModel(const FMessageAddress& ClientAddress)
-		: FTransferStatisticsModelBase()
-		, ClientAddress(ClientAddress)
-	{
-	}
+		: FTransferStatisticsModelBase(
+			TClientTransferStatTracker<FOutboundTransferStatistics>::FShouldInclude::CreateStatic(&FClientTransferStatisticsModel::ShouldIncludeOutboundStat, ClientAddress),
+			TClientTransferStatTracker<FInboundTransferStatistics>::FShouldInclude::CreateStatic(&FClientTransferStatisticsModel::ShouldIncludeInboundStat, ClientAddress)
+			)
+	{}
 
-	bool FClientTransferStatisticsModel::ShouldIncludeOutboundStat(const FOutboundTransferStatistics& Item) const
+	bool FClientTransferStatisticsModel::ShouldIncludeOutboundStat(const FOutboundTransferStatistics& Item, const FMessageAddress ClientAddress)
 	{
 		if (const INetworkMessagingExtension* Statistics = Private::ClientTransferStatisticsModel::GetMessagingStatistics())
 		{
@@ -54,7 +55,7 @@ namespace UE::MultiUserServer
 		return false;
 	}
 
-	bool FClientTransferStatisticsModel::ShouldIncludeInboundStat(const FInboundTransferStatistics& Item) const
+	bool FClientTransferStatisticsModel::ShouldIncludeInboundStat(const FInboundTransferStatistics& Item, const FMessageAddress ClientAddress)
 	{
 		if (const INetworkMessagingExtension* Statistics = Private::ClientTransferStatisticsModel::GetMessagingStatistics())
 		{
diff --git a/Engine/Plugins/Developer/Concert/ConcertApp/MultiUserServer/Source/MultiUserServer/Private/Widgets/Clients/Browser/Models/Transfer/ClientTransferStatisticsModel.h b/Engine/Plugins/Developer/Concert/ConcertApp/MultiUserServer/Source/MultiUserServer/Private/Widgets/Clients/Browser/Models/Transfer/ClientTransferStatisticsModel.h
--- a/Engine/Plugins/Developer/Concert/ConcertApp/MultiUserServer/Source/MultiUserServer/Private/Widgets/Clients/Browser/Models/Transfer/ClientTransferStatisticsModel.h
+++ b/Engine/Plugins/Developer/Concert/ConcertApp/MultiUserServer/Source/MultiUserServer/Private/Widgets/Clients/Browser/Models/Transfer/ClientTransferStatisticsModel.h
@@ -16,11 +16,9 @@ namespace UE::MultiUserServer
 		FClientTransferStatisticsModel(const FMessageAddress& ClientAddress);
 
 	protected:
-
-		const FMessageAddress ClientAddress;
 		
-		virtual bool ShouldIncludeOutboundStat(const FOutboundTransferStatistics& Item) const override;
-		virtual bool ShouldIncludeInboundStat(const FInboundTransferStatistics& Item) const override;
+		static bool ShouldIncludeOutboundStat(const FOutboundTransferStatistics& Item, const FMessageAddress ClientAddress);
+		static bool ShouldIncludeInboundStat(const FInboundTransferStatistics& Item, const FMessageAddress ClientAddress);
 	};
 }
 
diff --git a/Engine/Plugins/Developer/Concert/ConcertApp/MultiUserServer/Source/MultiUserServer/Private/Widgets/Clients/Browser/Models/Transfer/ServerTransferStatisticsModel.h b/Engine/Plugins/Developer/Concert/ConcertApp/MultiUserServer/Source/MultiUserServer/Private/Widgets/Clients/Browser/Models/Transfer/ServerTransferStatisticsModel.h
--- a/Engine/Plugins/Developer/Concert/ConcertApp/MultiUserServer/Source/MultiUserServer/Private/Widgets/Clients/Browser/Models/Transfer/ServerTransferStatisticsModel.h
+++ b/Engine/Plugins/Developer/Concert/ConcertApp/MultiUserServer/Source/MultiUserServer/Private/Widgets/Clients/Browser/Models/Transfer/ServerTransferStatisticsModel.h
@@ -10,10 +10,15 @@ namespace UE::MultiUserServer
 {
 	class FServerTransferStatisticsModel : public FTransferStatisticsModelBase
 	{
-	protected:
-		
-		virtual bool ShouldIncludeOutboundStat(const FOutboundTransferStatistics& Item) const override { return true; }
-		virtual bool ShouldIncludeInboundStat(const FInboundTransferStatistics& Item) const override { return true; }
+	public:
+
+		FServerTransferStatisticsModel()
+			:
+			FTransferStatisticsModelBase(
+			TClientTransferStatTracker<FOutboundTransferStatistics>::FShouldInclude::CreateLambda([](const auto&){ return true; }),
+			TClientTransferStatTracker<FInboundTransferStatistics>::FShouldInclude::CreateLambda([](const auto& ){ return true; })
+			)
+		{}
 	};
 }
 
diff --git a/Engine/Plugins/Developer/Concert/ConcertApp/MultiUserServer/Source/MultiUserServer/Private/Widgets/Clients/Browser/Models/Transfer/TransferStatisticsModel.cpp b/Engine/Plugins/Developer/Concert/ConcertApp/MultiUserServer/Source/MultiUserServer/Private/Widgets/Clients/Browser/Models/Transfer/TransferStatisticsModel.cpp
--- a/Engine/Plugins/Developer/Concert/ConcertApp/MultiUserServer/Source/MultiUserServer/Private/Widgets/Clients/Browser/Models/Transfer/TransferStatisticsModel.cpp
+++ b/Engine/Plugins/Developer/Concert/ConcertApp/MultiUserServer/Source/MultiUserServer/Private/Widgets/Clients/Browser/Models/Transfer/TransferStatisticsModel.cpp
@@ -37,17 +37,46 @@ namespace UE::MultiUserServer
 			ensureMsgf(false, TEXT("Feature %s is unavailable"), *INetworkMessagingExtension::ModularFeatureName.ToString());
 			return nullptr;
 		}
+
+		static TSharedRef<TClientTransferStatTracker<FOutboundTransferStatistics>> MakeOutbound(TClientTransferStatTracker<FOutboundTransferStatistics>::FShouldInclude ShouldIncludeOutboundFunc)
+		{
+			return MakeShared<TClientTransferStatTracker<FOutboundTransferStatistics>>(
+				MoveTemp(ShouldIncludeOutboundFunc),
+				TClientTransferStatTracker<FOutboundTransferStatistics>::FStatGetter::CreateLambda(
+					[](const FOutboundTransferStatistics& Stats)
+					{
+						return Stats.BytesSent;
+					}
+				)
+			);
+		}
+
+		static TSharedRef<TClientTransferStatTracker<FInboundTransferStatistics>> MakeInbound(TClientTransferStatTracker<FInboundTransferStatistics>::FShouldInclude ShouldIncludeInboundFunc)
+		{
+			return MakeShared<TClientTransferStatTracker<FInboundTransferStatistics>>(
+				MoveTemp(ShouldIncludeInboundFunc),
+				TClientTransferStatTracker<FInboundTransferStatistics>::FStatGetter::CreateLambda(
+					[](const FInboundTransferStatistics& Stats)
+					{
+						return Stats.BytesReceived;
+					}
+				)
+			);
+		}
 	}
 
-	FTransferStatisticsModelBase::FTransferStatisticsModelBase()
-		: OutboundStatTracker([this](const FOutboundTransferStatistics& Stats) { return ShouldIncludeOutboundStat(Stats); }, [](const FOutboundTransferStatistics& Stats){ return Stats.BytesSent; })
-		, InboundStatTracker([this](const FInboundTransferStatistics& Stats) { return ShouldIncludeInboundStat(Stats); }, [](const FInboundTransferStatistics& Stats){ return Stats.BytesReceived; })
+	FTransferStatisticsModelBase::FTransferStatisticsModelBase(
+		TClientTransferStatTracker<FOutboundTransferStatistics>::FShouldInclude ShouldIncludeOutboundFunc,
+		TClientTransferStatTracker<FInboundTransferStatistics>::FShouldInclude ShouldIncludeInboundFunc)
+		: OutboundStatTracker(Private::TransferStatisticsModel::MakeOutbound(MoveTemp(ShouldIncludeOutboundFunc)))
+		, InboundStatTracker(Private::TransferStatisticsModel::MakeInbound(MoveTemp(ShouldIncludeInboundFunc)))
 	{
 		if (INetworkMessagingExtension* Statistics = Private::TransferStatisticsModel::GetMessagingStatistics())
 		{
-			Statistics->OnOutboundTransferUpdatedFromThread().AddRaw(&OutboundStatTracker, &TClientTransferStatTracker<FOutboundTransferStatistics>::OnTransferUpdatedFromThread);
-			Statistics->OnInboundTransferUpdatedFromThread().AddRaw(&InboundStatTracker, &TClientTransferStatTracker<FInboundTransferStatistics>::OnTransferUpdatedFromThread);
-			TickHandle = FTSTicker::GetCoreTicker().AddTicker(FTickerDelegate::CreateRaw(this, &FTransferStatisticsModelBase::Tick));
+			Statistics->OnOutboundTransferUpdatedFromThread().AddSP(OutboundStatTracker, &TClientTransferStatTracker<FOutboundTransferStatistics>::OnTransferUpdatedFromThread);
+			Statistics->OnInboundTransferUpdatedFromThread().AddSP(InboundStatTracker, &TClientTransferStatTracker<FInboundTransferStatistics>::OnTransferUpdatedFromThread);
+			InboundTickHandle = FTSTicker::GetCoreTicker().AddTicker(FTickerDelegate::CreateSP(OutboundStatTracker, &TClientTransferStatTracker<FOutboundTransferStatistics>::Tick));
+			OutboundTickHandle = FTSTicker::GetCoreTicker().AddTicker(FTickerDelegate::CreateSP(InboundStatTracker, &TClientTransferStatTracker<FInboundTransferStatistics>::Tick));
 		}
 	}
 
@@ -59,19 +88,20 @@ namespace UE::MultiUserServer
 			Statistics->OnInboundTransferUpdatedFromThread().RemoveAll(&InboundStatTracker);
 		}
 		
-		FTSTicker::GetCoreTicker().RemoveTicker(TickHandle);
+		FTSTicker::GetCoreTicker().RemoveTicker(InboundTickHandle);
+		FTSTicker::GetCoreTicker().RemoveTicker(OutboundTickHandle);
 	}
 
 	const TArray<FConcertTransferSamplePoint>& FTransferStatisticsModelBase::GetTransferStatTimeline(EConcertTransferStatistic StatisticType) const
 	{
 		switch (StatisticType)
 		{
-		case EConcertTransferStatistic::SentToClient: return OutboundStatTracker.GetTransferStatisticsTimeline();
-		case EConcertTransferStatistic::ReceivedFromClient: return InboundStatTracker.GetTransferStatisticsTimeline();;
+		case EConcertTransferStatistic::SentToClient: return OutboundStatTracker->GetTransferStatisticsTimeline();
+		case EConcertTransferStatistic::ReceivedFromClient: return InboundStatTracker->GetTransferStatisticsTimeline();;
 		case EConcertTransferStatistic::Count:
 		default:
 			checkNoEntry();
-			return OutboundStatTracker.GetTransferStatisticsTimeline();
+			return OutboundStatTracker->GetTransferStatisticsTimeline();
 		}
 	}
 
@@ -79,19 +109,12 @@ namespace UE::MultiUserServer
 	{
 		switch (StatisticType)
 		{
-		case EConcertTransferStatistic::SentToClient: return OutboundStatTracker.GetOnTimelineUpdatedDelegates();
-		case EConcertTransferStatistic::ReceivedFromClient: return InboundStatTracker.GetOnTimelineUpdatedDelegates();;
+		case EConcertTransferStatistic::SentToClient: return OutboundStatTracker->GetOnTimelineUpdatedDelegates();
+		case EConcertTransferStatistic::ReceivedFromClient: return InboundStatTracker->GetOnTimelineUpdatedDelegates();;
 		case EConcertTransferStatistic::Count:
 		default:
 			checkNoEntry();
-			return OutboundStatTracker.GetOnTimelineUpdatedDelegates();
+			return OutboundStatTracker->GetOnTimelineUpdatedDelegates();
 		}
 	}
-	
-	bool FTransferStatisticsModelBase::Tick(float DeltaTime)
-	{
-		OutboundStatTracker.Tick();
-		InboundStatTracker.Tick();
-		return true;
-	}
 }
\ No newline at end of file
diff --git a/Engine/Plugins/Developer/Concert/ConcertApp/MultiUserServer/Source/MultiUserServer/Private/Widgets/Clients/Browser/Models/Transfer/TransferStatisticsModel.h b/Engine/Plugins/Developer/Concert/ConcertApp/MultiUserServer/Source/MultiUserServer/Private/Widgets/Clients/Browser/Models/Transfer/TransferStatisticsModel.h
--- a/Engine/Plugins/Developer/Concert/ConcertApp/MultiUserServer/Source/MultiUserServer/Private/Widgets/Clients/Browser/Models/Transfer/TransferStatisticsModel.h
+++ b/Engine/Plugins/Developer/Concert/ConcertApp/MultiUserServer/Source/MultiUserServer/Private/Widgets/Clients/Browser/Models/Transfer/TransferStatisticsModel.h
@@ -33,8 +33,8 @@ namespace UE::MultiUserServer
 		
 	public:
 
-		using FShouldInclude = TFunction<bool(const TStatBase&)>;
-		using FStatGetter = TFunction<uint64(const TStatBase&)>;
+		DECLARE_DELEGATE_RetVal_OneParam(bool, FShouldInclude, const TStatBase&);
+		DECLARE_DELEGATE_RetVal_OneParam(uint64, FStatGetter, const TStatBase&);
 
 		TClientTransferStatTracker(FShouldInclude ShouldIncludeFunc, FStatGetter StatGetterFunc)
 			: ShouldIncludeFunc(MoveTemp(ShouldIncludeFunc))
@@ -43,13 +43,13 @@ namespace UE::MultiUserServer
 
 		void OnTransferUpdatedFromThread(TStatBase TransferStatistics)
 		{
-			if (ShouldIncludeFunc(TransferStatistics))
+			if (ShouldIncludeFunc.Execute(TransferStatistics))
 			{
 				AsyncStatQueue.Enqueue(TStatItem{ FDateTime::Now(), TransferStatistics });
 			}
 		}
 
-		void Tick()
+		bool Tick(float DeltaTime)
 		{
 			const bool bRemovedOldStats = RemoveOldStatTimelines();
 			const bool bRemovedAnyGroups = RemoveOldGroupedStats();
@@ -70,6 +70,8 @@ namespace UE::MultiUserServer
 				
 				OnGroupsUpdatedDelegate.Broadcast();
 			}
+
+			return true;
 		}
 		
 		FOnTransferTimelineUpdated& GetOnTimelineUpdatedDelegates() { return OnTimelineUpdatedDelegates; }
@@ -143,14 +145,14 @@ namespace UE::MultiUserServer
 		
 		void UpdateStatTimeline(const TStatBase& TransferStatistics, FConcertTransferSamplePoint& SamplingThisTick)
 		{
-			uint64 BytesSentSinceLastTime = StatGetterFunc(TransferStatistics);
+			uint64 BytesSentSinceLastTime = StatGetterFunc.Execute(TransferStatistics);
 			// Same message ID may have been queued multiple times
 			if (FIncompleteMessageData* ItemThisTick = IncompleteStatsUntilNow.Find(TransferStatistics.MessageId))
 			{
 				BytesSentSinceLastTime -= ItemThisTick->BytesTransferredSoFar;
 			}
 
-			IncompleteStatsUntilNow.Add(TransferStatistics.MessageId, { StatGetterFunc(TransferStatistics), FDateTime::Now() });
+			IncompleteStatsUntilNow.Add(TransferStatistics.MessageId, { StatGetterFunc.Execute(TransferStatistics), FDateTime::Now() });
 			SamplingThisTick.BytesTransferred += BytesSentSinceLastTime;
 		}
 		
@@ -178,28 +180,29 @@ namespace UE::MultiUserServer
 	{
 	public:
 		
-		FTransferStatisticsModelBase();
+		FTransferStatisticsModelBase(
+			TClientTransferStatTracker<FOutboundTransferStatistics>::FShouldInclude ShouldIncludeOutboundFunc,
+			TClientTransferStatTracker<FInboundTransferStatistics>::FShouldInclude ShouldIncludeInboundFunc
+			);
 		virtual ~FTransferStatisticsModelBase() override;
 		
 		virtual const TArray<FConcertTransferSamplePoint>& GetTransferStatTimeline(EConcertTransferStatistic StatisticType) const override;
-		virtual const TArray<TSharedPtr<FOutboundTransferStatistics>>& GetOutboundTransferStatsGroupedById() const override { return OutboundStatTracker.GetTransferStatisticsGroupedById(); }
-		virtual const TArray<TSharedPtr<FInboundTransferStatistics>>& GetInboundTransferStatsGroupedById() const override { return InboundStatTracker.GetTransferStatisticsGroupedById(); }
+		virtual const TArray<TSharedPtr<FOutboundTransferStatistics>>& GetOutboundTransferStatsGroupedById() const override { return OutboundStatTracker->GetTransferStatisticsGroupedById(); }
+		virtual const TArray<TSharedPtr<FInboundTransferStatistics>>& GetInboundTransferStatsGroupedById() const override { return InboundStatTracker->GetTransferStatisticsGroupedById(); }
 		virtual FOnTransferTimelineUpdated& OnTransferTimelineUpdated(EConcertTransferStatistic StatisticType) override;
-		virtual FOnTransferGroupsUpdated& OnOutboundTransferGroupsUpdated() override { return OutboundStatTracker.GetOnGroupsUpdatedDelegate(); }
-		virtual FOnTransferGroupsUpdated& OnInboundTransferGroupsUpdated() override { return InboundStatTracker.GetOnGroupsUpdatedDelegate(); }
+		virtual FOnTransferGroupsUpdated& OnOutboundTransferGroupsUpdated() override { return OutboundStatTracker->GetOnGroupsUpdatedDelegate(); }
+		virtual FOnTransferGroupsUpdated& OnInboundTransferGroupsUpdated() override { return InboundStatTracker->GetOnGroupsUpdatedDelegate(); }
 
 	private:
-		
-		virtual bool ShouldIncludeOutboundStat(const FOutboundTransferStatistics& Item) const = 0;
-		virtual bool ShouldIncludeInboundStat(const FInboundTransferStatistics& Item) const = 0;
 
-		FTSTicker::FDelegateHandle TickHandle;
-		
+		FTSTicker::FDelegateHandle InboundTickHandle;
+		FTSTicker::FDelegateHandle OutboundTickHandle;
 		
-		bool Tick(float DeltaTime);
 
-		TClientTransferStatTracker<FOutboundTransferStatistics> OutboundStatTracker;
-		TClientTransferStatTracker<FInboundTransferStatistics> InboundStatTracker;
+		// These are used for subscribing to INetworkMessagingExtension - therefore they are shared pointers for synchronization purposes
+		// Prevents race condition when destructor is called and UDP thread attempts to call the callbacks
+		TSharedRef<TClientTransferStatTracker<FOutboundTransferStatistics>> OutboundStatTracker;
+		TSharedRef<TClientTransferStatTracker<FInboundTransferStatistics>> InboundStatTracker;
 	};
 }
 
