Verify the map has the message type name.

diff --git a/Engine/Plugins/Developer/Concert/ConcertApp/MultiUserServer/Source/MultiUserServer/Private/Widgets/Clients/Logging/Util/MessageTypeUtils.cpp b/Engine/Plugins/Developer/Concert/ConcertApp/MultiUserServer/Source/MultiUserServer/Private/Widgets/Clients/Logging/Util/MessageTypeUtils.cpp
--- a/Engine/Plugins/Developer/Concert/ConcertApp/MultiUserServer/Source/MultiUserServer/Private/Widgets/Clients/Logging/Util/MessageTypeUtils.cpp
+++ b/Engine/Plugins/Developer/Concert/ConcertApp/MultiUserServer/Source/MultiUserServer/Private/Widgets/Clients/Logging/Util/MessageTypeUtils.cpp
@@ -85,10 +85,10 @@ namespace UE::MultiUserServer::MessageTypeUtils
 			}
 			return Result;
 		}();
-
-		return CachedSanitizedNames[MessageTypeName];
+		const FString* SanitizedName = CachedSanitizedNames.Find(MessageTypeName);
+		return SanitizedName ? *SanitizedName : "UNKNOWN";
 	}
-	
+
 	static TSet<FName> FilterAllMessageTypes(TFunctionRef<bool(UStruct*)> ConsumerFunc)
 	{
 		TSet<FName> Result;
@@ -102,4 +102,4 @@ namespace UE::MultiUserServer::MessageTypeUtils
 		}
 		return Result;
 	}
-};
\ No newline at end of file
+};
