Check for nullptr before attempting to save the layout.

#rnx

diff --git a/Engine/Plugins/Developer/Concert/ConcertApp/MultiUserServer/Source/MultiUserServer/Private/Window/ConcertServerWindowController.cpp b/Engine/Plugins/Developer/Concert/ConcertApp/MultiUserServer/Source/MultiUserServer/Private/Window/ConcertServerWindowController.cpp
--- a/Engine/Plugins/Developer/Concert/ConcertApp/MultiUserServer/Source/MultiUserServer/Private/Window/ConcertServerWindowController.cpp
+++ b/Engine/Plugins/Developer/Concert/ConcertApp/MultiUserServer/Source/MultiUserServer/Private/Window/ConcertServerWindowController.cpp
@@ -176,8 +176,11 @@ namespace UE::MultiUserServer
 
 	void FConcertServerWindowController::SaveLayout() const
 	{
-		FLayoutSaveRestore::SaveToConfig(MultiUserServerLayoutIni, PersistentLayout.ToSharedRef());
-	    GConfig->Flush(false, MultiUserServerLayoutIni);
+		if (PersistentLayout)
+		{
+			FLayoutSaveRestore::SaveToConfig(MultiUserServerLayoutIni, PersistentLayout.ToSharedRef());
+		    GConfig->Flush(false, MultiUserServerLayoutIni);
+		}
 	}
 }
 
