Clean-up header for do not show again dialog.

diff --git a/Engine/Plugins/Developer/Concert/ConcertApp/MultiUserServer/Source/MultiUserServer/Private/Widgets/Util/SDoNotShowAgainDialog.h b/Engine/Plugins/Developer/Concert/ConcertApp/MultiUserServer/Source/MultiUserServer/Private/Widgets/Util/SDoNotShowAgainDialog.h
--- a/Engine/Plugins/Developer/Concert/ConcertApp/MultiUserServer/Source/MultiUserServer/Private/Widgets/Util/SDoNotShowAgainDialog.h
+++ b/Engine/Plugins/Developer/Concert/ConcertApp/MultiUserServer/Source/MultiUserServer/Private/Widgets/Util/SDoNotShowAgainDialog.h
@@ -4,15 +4,9 @@
 
 #include "Dialog/SCustomDialog.h"
 
-class FArchivedSessionHistoryController;
 class IConcertSyncServer;
 class SCheckBox;
 
-namespace UE::ConcertSyncCore
-{
-	struct FHistoryAnalysisResult;
-}
-
 namespace UE::MultiUserServer
 {
 	/** Text message with a "Do not show again" checkbox. */
