Fixes to accommodate removal of template specialization TOptional<TNonNullPtr<OptionalType>> and replacement with the IntrusiveUnsetOptionalState mechanism instead

#rnx

diff --git a/Engine/Plugins/Developer/Concert/ConcertApp/MultiUserServer/Source/MultiUserServer/Private/Widgets/Clients/PackageTransmission/Model/PackageTransmissionModel.cpp b/Engine/Plugins/Developer/Concert/ConcertApp/MultiUserServer/Source/MultiUserServer/Private/Widgets/Clients/PackageTransmission/Model/PackageTransmissionModel.cpp
--- a/Engine/Plugins/Developer/Concert/ConcertApp/MultiUserServer/Source/MultiUserServer/Private/Widgets/Clients/PackageTransmission/Model/PackageTransmissionModel.cpp
+++ b/Engine/Plugins/Developer/Concert/ConcertApp/MultiUserServer/Source/MultiUserServer/Private/Widgets/Clients/PackageTransmission/Model/PackageTransmissionModel.cpp
@@ -110,7 +110,7 @@ namespace UE::MultiUserServer
 		}
 
 		int64 PackageRevision;
-		const bool bCouldGetHeadRevision = Database->GetPackageHeadRevision(Args.PackageInfo.PackageName, PackageRevision);
+		const bool bCouldGetHeadRevision = (*Database)->GetPackageHeadRevision(Args.PackageInfo.PackageName, PackageRevision);
 		if (ensureMsgf(bCouldGetHeadRevision, TEXT("The activity should have been processed properly before this event was emitted")))
 		{
 			return SendDirection == EPackageSendDirection::ClientToServer
diff --git a/Engine/Plugins/Developer/Concert/ConcertApp/MultiUserServer/Source/MultiUserServer/Private/Widgets/HistoryEdition/SActivityDependencyDialog.cpp b/Engine/Plugins/Developer/Concert/ConcertApp/MultiUserServer/Source/MultiUserServer/Private/Widgets/HistoryEdition/SActivityDependencyDialog.cpp
--- a/Engine/Plugins/Developer/Concert/ConcertApp/MultiUserServer/Source/MultiUserServer/Private/Widgets/HistoryEdition/SActivityDependencyDialog.cpp
+++ b/Engine/Plugins/Developer/Concert/ConcertApp/MultiUserServer/Source/MultiUserServer/Private/Widgets/HistoryEdition/SActivityDependencyDialog.cpp
@@ -72,10 +72,10 @@ namespace UE::MultiUserServer
 		const TOptional<FConcertSyncSessionDatabaseNonNullPtr> SessionDatabase = InSyncServer->GetArchivedSessionDatabase(InSessionId);
 		check(SessionDatabase);
 		
-		ConcertSyncCore::FActivityDependencyGraph DependencyGraph = ConcertSyncCore::BuildDependencyGraphFrom(*SessionDatabase);
+		ConcertSyncCore::FActivityDependencyGraph DependencyGraph = ConcertSyncCore::BuildDependencyGraphFrom(**SessionDatabase);
 		if (ConsoleVariables::CVarLogActivityDependencyGraphOnDelete.GetValueOnGameThread())
 		{
-			UE_LOG(LogConcert, Log, TEXT("Analysed graph: %s"), *UE::ConcertSyncCore::Graphviz::ExportToGraphviz(DependencyGraph, *SessionDatabase));
+			UE_LOG(LogConcert, Log, TEXT("Analysed graph: %s"), *UE::ConcertSyncCore::Graphviz::ExportToGraphviz(DependencyGraph, **SessionDatabase));
 		}
 		
 		return SAssignNew(DependencyView, SActivityDependencyView, MoveTemp(DependencyGraph))
diff --git a/Engine/Plugins/Developer/Concert/ConcertApp/MultiUserServer/Source/MultiUserServer/Private/Widgets/HistoryEdition/SMuteActivityDependenciesDialog.cpp b/Engine/Plugins/Developer/Concert/ConcertApp/MultiUserServer/Source/MultiUserServer/Private/Widgets/HistoryEdition/SMuteActivityDependenciesDialog.cpp
--- a/Engine/Plugins/Developer/Concert/ConcertApp/MultiUserServer/Source/MultiUserServer/Private/Widgets/HistoryEdition/SMuteActivityDependenciesDialog.cpp
+++ b/Engine/Plugins/Developer/Concert/ConcertApp/MultiUserServer/Source/MultiUserServer/Private/Widgets/HistoryEdition/SMuteActivityDependenciesDialog.cpp
@@ -107,7 +107,7 @@ namespace UE::MultiUserServer
 		if (TSharedPtr<IConcertSyncServer> PinnedSyncServer = SyncServer.Pin())
 		{
 			FConcertSyncActivity Activity;
-			const bool bGotActivity = PinnedSyncServer->GetArchivedSessionDatabase(SessionId)->GetActivity(ActivityId, Activity);
+			const bool bGotActivity = (*PinnedSyncServer->GetArchivedSessionDatabase(SessionId))->GetActivity(ActivityId, Activity);
 			check(bGotActivity);
 			
 			const bool bIsMuted = (Activity.Flags & EConcertSyncActivityFlags::Muted) != EConcertSyncActivityFlags::None;
diff --git a/Engine/Plugins/Developer/Concert/ConcertApp/MultiUserServer/Source/MultiUserServer/Private/Widgets/SessionTabs/Archived/ArchivedConcertSessionTab.cpp b/Engine/Plugins/Developer/Concert/ConcertApp/MultiUserServer/Source/MultiUserServer/Private/Widgets/SessionTabs/Archived/ArchivedConcertSessionTab.cpp
--- a/Engine/Plugins/Developer/Concert/ConcertApp/MultiUserServer/Source/MultiUserServer/Private/Widgets/SessionTabs/Archived/ArchivedConcertSessionTab.cpp
+++ b/Engine/Plugins/Developer/Concert/ConcertApp/MultiUserServer/Source/MultiUserServer/Private/Widgets/SessionTabs/Archived/ArchivedConcertSessionTab.cpp
@@ -357,7 +357,7 @@ void FArchivedConcertSessionTab::SpawnDialogForMutingOrUnmuting(bool bShouldMute
 				// Because the dialog is non-modal, the user may have closed the program in the mean time
 				if (const TSharedPtr<const FArchivedConcertSessionTab> PinnedThis = WeakTabThis.Pin())
 				{
-					const bool bSuccess = PinnedThis->SyncServer->GetArchivedSessionDatabase(PinnedThis->InspectedSessionID)->SetActivities(Selection, [bShouldMute](FConcertSyncActivity& Activity)
+					const bool bSuccess = (*PinnedThis->SyncServer->GetArchivedSessionDatabase(PinnedThis->InspectedSessionID))->SetActivities(Selection, [bShouldMute](FConcertSyncActivity& Activity)
 					{
 						if (bShouldMute)
 						{
diff --git a/Engine/Plugins/Developer/Concert/ConcertApp/MultiUserServer/Source/MultiUserServer/Private/Widgets/SessionTabs/Live/LiveServerSessionHistoryController.cpp b/Engine/Plugins/Developer/Concert/ConcertApp/MultiUserServer/Source/MultiUserServer/Private/Widgets/SessionTabs/Live/LiveServerSessionHistoryController.cpp
--- a/Engine/Plugins/Developer/Concert/ConcertApp/MultiUserServer/Source/MultiUserServer/Private/Widgets/SessionTabs/Live/LiveServerSessionHistoryController.cpp
+++ b/Engine/Plugins/Developer/Concert/ConcertApp/MultiUserServer/Source/MultiUserServer/Private/Widgets/SessionTabs/Live/LiveServerSessionHistoryController.cpp
@@ -30,7 +30,7 @@ FLiveServerSessionHistoryController::FLiveServerSessionHistoryController(TShared
 
 	if (const TOptional<FConcertSyncSessionDatabaseNonNullPtr> Database = SyncServer->GetLiveSessionDatabase(GetSessionId()))
 	{
-		Database->OnActivityProduced().AddRaw(this, &FLiveServerSessionHistoryController::OnSessionProduced);
+		(*Database)->OnActivityProduced().AddRaw(this, &FLiveServerSessionHistoryController::OnSessionProduced);
 	}
 }
 
@@ -42,7 +42,7 @@ FLiveServerSessionHistoryController::~FLiveServerSessionHistoryController()
 	}
 	if (const TOptional<FConcertSyncSessionDatabaseNonNullPtr> Database = SyncServer->GetLiveSessionDatabase(GetSessionId()))
 	{
-		Database->OnActivityProduced().RemoveAll(this);
+		(*Database)->OnActivityProduced().RemoveAll(this);
 	}
 }
 
diff --git a/Engine/Plugins/Developer/Concert/ConcertApp/MultiUserServer/Source/MultiUserServer/Private/Widgets/SessionTabs/Live/PackageViewer/ConcertSessionPackageViewerController.cpp b/Engine/Plugins/Developer/Concert/ConcertApp/MultiUserServer/Source/MultiUserServer/Private/Widgets/SessionTabs/Live/PackageViewer/ConcertSessionPackageViewerController.cpp
--- a/Engine/Plugins/Developer/Concert/ConcertApp/MultiUserServer/Source/MultiUserServer/Private/Widgets/SessionTabs/Live/PackageViewer/ConcertSessionPackageViewerController.cpp
+++ b/Engine/Plugins/Developer/Concert/ConcertApp/MultiUserServer/Source/MultiUserServer/Private/Widgets/SessionTabs/Live/PackageViewer/ConcertSessionPackageViewerController.cpp
@@ -20,7 +20,7 @@ FConcertSessionPackageViewerController::FConcertSessionPackageViewerController(T
 	
 	if (const TOptional<FConcertSyncSessionDatabaseNonNullPtr> Database = SyncServer->GetLiveSessionDatabase(InspectedSession->GetId()))
 	{
-		Database->OnActivityProduced().AddRaw(this, &FConcertSessionPackageViewerController::OnSessionProduced);
+		(*Database)->OnActivityProduced().AddRaw(this, &FConcertSessionPackageViewerController::OnSessionProduced);
 	}
 }
 
@@ -33,7 +33,7 @@ FConcertSessionPackageViewerController::~FConcertSessionPackageViewerController(
 	
 	if (const TOptional<FConcertSyncSessionDatabaseNonNullPtr> Database = SyncServer->GetLiveSessionDatabase(InspectedSession->GetId()))
 	{
-		Database->OnActivityProduced().RemoveAll(this);
+		(*Database)->OnActivityProduced().RemoveAll(this);
 	}
 }
 
@@ -50,7 +50,7 @@ void FConcertSessionPackageViewerController::ReloadActivities() const
 		};
 		
 		TMap<FName, FPackageActivity> LatestPackageActivities;
-		Database->EnumeratePackageActivities([this, &LatestPackageActivities](FConcertSyncActivity&& BasePart, FConcertSyncPackageEventData& EventData)
+		(*Database)->EnumeratePackageActivities([this, &LatestPackageActivities](FConcertSyncActivity&& BasePart, FConcertSyncPackageEventData& EventData)
 		{
 			FStructOnScope ActivitySummary;
 			if (BasePart.EventSummary.GetPayload(ActivitySummary))
@@ -89,7 +89,7 @@ TOptional<FConcertClientInfo> FConcertSessionPackageViewerController::GetClientI
 {
 	const TOptional<FConcertSyncSessionDatabaseNonNullPtr> Database = SyncServer->GetLiveSessionDatabase(InspectedSession->GetId());
 	FConcertSyncEndpointData Result;
-	if (Database && Database->GetEndpoint(ClientId, Result))
+	if (Database && (*Database)->GetEndpoint(ClientId, Result))
 	{
 		return Result.ClientInfo;
 	}
@@ -100,7 +100,7 @@ bool FConcertSessionPackageViewerController::GetPackageEvent(const FConcertSessi
 {
 	if (const TOptional<FConcertSyncSessionDatabaseNonNullPtr> Database = SyncServer->GetLiveSessionDatabase(InspectedSession->GetId()))
 	{
-		return Database->GetPackageEventMetaData(Activity.Activity.EventId, OutPackageEvent.PackageRevision, OutPackageEvent.PackageInfo);
+		return (*Database)->GetPackageEventMetaData(Activity.Activity.EventId, OutPackageEvent.PackageRevision, OutPackageEvent.PackageInfo);
 	}
 	return false;
 }
@@ -115,13 +115,13 @@ TOptional<int64> FConcertSessionPackageViewerController::GetSizeOfPackageEvent(c
 
 	const TSet<EConcertPackageUpdateType> PackageTypesWithSize { EConcertPackageUpdateType::Added, EConcertPackageUpdateType::Renamed, EConcertPackageUpdateType::Saved };
 	FConcertSyncPackageEventMetaData PackageEventMetaData;
-	if (Database->GetPackageEventMetaData(Activity.Activity.EventId, PackageEventMetaData.PackageRevision, PackageEventMetaData.PackageInfo)
+	if ((*Database)->GetPackageEventMetaData(Activity.Activity.EventId, PackageEventMetaData.PackageRevision, PackageEventMetaData.PackageInfo)
 		; PackageTypesWithSize.Contains(PackageEventMetaData.PackageInfo.PackageUpdateType))
 	{
 		const FName PackageName = PackageEventMetaData.PackageInfo.PackageUpdateType == EConcertPackageUpdateType::Renamed
 			?  PackageEventMetaData.PackageInfo.NewPackageName
 			: PackageEventMetaData.PackageInfo.PackageName;
-		return Database->GetPackageSizeForRevision(PackageName);
+		return (*Database)->GetPackageSizeForRevision(PackageName);
 	}
 	return {};
 }
diff --git a/Engine/Plugins/Developer/Concert/ConcertApp/MultiUserServer/Source/MultiUserServer/Private/Widgets/SessionTabs/ServerSessionHistoryControllerBase.cpp b/Engine/Plugins/Developer/Concert/ConcertApp/MultiUserServer/Source/MultiUserServer/Private/Widgets/SessionTabs/ServerSessionHistoryControllerBase.cpp
--- a/Engine/Plugins/Developer/Concert/ConcertApp/MultiUserServer/Source/MultiUserServer/Private/Widgets/SessionTabs/ServerSessionHistoryControllerBase.cpp
+++ b/Engine/Plugins/Developer/Concert/ConcertApp/MultiUserServer/Source/MultiUserServer/Private/Widgets/SessionTabs/ServerSessionHistoryControllerBase.cpp
@@ -20,14 +20,14 @@ void FServerSessionHistoryControllerBase::GetActivities(int64 MaximumNumberOfAct
 		OutFetchedActivities.Reset();
 
 		int64 LastActivityId = INDEX_NONE;
-		Database->GetActivityMaxId(LastActivityId);
+		(*Database)->GetActivityMaxId(LastActivityId);
 		const int64 FirstActivityIdToFetch = FMath::Max<int64>(1, LastActivityId - MaximumNumberOfActivities);
-		Database->EnumerateActivitiesInRange(FirstActivityIdToFetch, MaximumNumberOfActivities, [this, &Database, &OutEndpointClientInfoMap, &OutFetchedActivities](FConcertSyncActivity&& InActivity)
+		(*Database)->EnumerateActivitiesInRange(FirstActivityIdToFetch, MaximumNumberOfActivities, [this, &Database, &OutEndpointClientInfoMap, &OutFetchedActivities](FConcertSyncActivity&& InActivity)
 		{
 			if (!OutEndpointClientInfoMap.Contains(InActivity.EndpointId))
 			{
 				FConcertSyncEndpointData EndpointData;
-				if (Database->GetEndpoint(InActivity.EndpointId, EndpointData))
+				if ((*Database)->GetEndpoint(InActivity.EndpointId, EndpointData))
 				{
 					OutEndpointClientInfoMap.Add(InActivity.EndpointId, EndpointData.ClientInfo);
 				}
@@ -48,7 +48,7 @@ bool FServerSessionHistoryControllerBase::GetPackageEvent(const FConcertSessionA
 {
 	if (const TOptional<FConcertSyncSessionDatabaseNonNullPtr> Database = GetSessionDatabase(SessionId))
 	{
-		return Database->GetPackageEventMetaData(Activity.Activity.EventId, OutPackageEvent.PackageRevision, OutPackageEvent.PackageInfo);
+		return (*Database)->GetPackageEventMetaData(Activity.Activity.EventId, OutPackageEvent.PackageRevision, OutPackageEvent.PackageInfo);
 	}
 	
 	return false;
@@ -58,7 +58,7 @@ TFuture<TOptional<FConcertSyncTransactionEvent>> FServerSessionHistoryController
 {
 	if (const TOptional<FConcertSyncSessionDatabaseNonNullPtr> Database = GetSessionDatabase(SessionId))
 	{
-		return FindOrRequestTransactionEvent(*Database, Activity.Activity.EventId);
+		return FindOrRequestTransactionEvent(**Database, Activity.Activity.EventId);
 	}
 	
 	return MakeFulfilledPromise<TOptional<FConcertSyncTransactionEvent>>().GetFuture(); // Not found.
diff --git a/Engine/Plugins/VirtualProduction/LevelSnapshots/Source/LevelSnapshots/Private/Data/Util/ClassDataUtil.cpp b/Engine/Plugins/VirtualProduction/LevelSnapshots/Source/LevelSnapshots/Private/Data/Util/ClassDataUtil.cpp
--- a/Engine/Plugins/VirtualProduction/LevelSnapshots/Source/LevelSnapshots/Private/Data/Util/ClassDataUtil.cpp
+++ b/Engine/Plugins/VirtualProduction/LevelSnapshots/Source/LevelSnapshots/Private/Data/Util/ClassDataUtil.cpp
@@ -32,7 +32,7 @@ namespace UE::LevelSnapshots::Private
 			return {};
 		}
 
-		UClass* Class = ArchetypeData->ClassPath.TryLoadClass<UObject>();
+		UClass* Class = (*ArchetypeData)->ClassPath.TryLoadClass<UObject>();
 		if (!Class)
 		{
 			return {};
@@ -113,7 +113,7 @@ namespace UE::LevelSnapshots::Private
 	{
 		if (const TOptional<TNonNullPtr<FClassSnapshotData>> ArchetypeData = GetObjectArchetypeData(WorldData, ClassIndex, Cache, FallbackInfo))
 		{
-			SerializeClassDefaultsInto(Object, *ArchetypeData, WorldData);
+			SerializeClassDefaultsInto(Object, **ArchetypeData, WorldData);
 		}
 		else
 		{
@@ -137,7 +137,7 @@ namespace UE::LevelSnapshots::Private
 	{
 		if (const TOptional<TNonNullPtr<FClassSnapshotData>> ArchetypeData = GetObjectArchetypeData(WorldData, ClassIndex, Cache, FallbackInfo))
 		{
-			FApplyClassDefaulDataArchive::RestoreSelectedChangedClassDefaults(*ArchetypeData, WorldData, Object, PropertiesToRestore);
+			FApplyClassDefaulDataArchive::RestoreSelectedChangedClassDefaults(**ArchetypeData, WorldData, Object, PropertiesToRestore);
 		}
 		else
 		{
diff --git a/Engine/Plugins/VirtualProduction/LevelSnapshots/Source/LevelSnapshots/Private/Data/Util/ObjectDependencyUtil.cpp b/Engine/Plugins/VirtualProduction/LevelSnapshots/Source/LevelSnapshots/Private/Data/Util/ObjectDependencyUtil.cpp
--- a/Engine/Plugins/VirtualProduction/LevelSnapshots/Source/LevelSnapshots/Private/Data/Util/ObjectDependencyUtil.cpp
+++ b/Engine/Plugins/VirtualProduction/LevelSnapshots/Source/LevelSnapshots/Private/Data/Util/ObjectDependencyUtil.cpp
@@ -87,8 +87,8 @@ namespace UE::LevelSnapshots::Private::Internal
 			return nullptr;
 		}
 		
-		const UClass* Class = Archetype->GetClass();
-		UObject* Subobject = NewObject<UObject>(SubobjectOuter, Class, SubobjectName, SubobjectData->GetObjectFlags(), Archetype.GetValue());
+		const UClass* Class = (*Archetype)->GetClass();
+		UObject* Subobject = NewObject<UObject>(SubobjectOuter, Class, SubobjectName, SubobjectData->GetObjectFlags(), Archetype->Get());
 		FSubobjectSnapshotCache& SubobjectCache = Cache.SubobjectCache.FindOrAdd(WorldData.SerializedObjectReferences[ObjectPathIndex]);
 		SubobjectCache.SnapshotObject = Subobject;
 		FLoadSnapshotObjectArchive::ApplyToSnapshotWorldObject(*SubobjectData, WorldData, Cache, Subobject, ProcessObjectDependency, LocalisationNamespace)	;
diff --git a/Engine/Plugins/VirtualProduction/LevelSnapshots/Source/LevelSnapshots/Private/Filtering/Diffing/EquivalenceUtil.cpp b/Engine/Plugins/VirtualProduction/LevelSnapshots/Source/LevelSnapshots/Private/Filtering/Diffing/EquivalenceUtil.cpp
--- a/Engine/Plugins/VirtualProduction/LevelSnapshots/Source/LevelSnapshots/Private/Filtering/Diffing/EquivalenceUtil.cpp
+++ b/Engine/Plugins/VirtualProduction/LevelSnapshots/Source/LevelSnapshots/Private/Filtering/Diffing/EquivalenceUtil.cpp
@@ -395,7 +395,7 @@ namespace UE::LevelSnapshots::Private::Internal
 		}
 
 		// Are the two subobjects owned by equivalent actors
-		const AActor* EquivalentSnapshotActor = UE::LevelSnapshots::Private::GetPreallocatedIfCached(OriginalOwningActor, Cache).Get(nullptr);
+		const AActor* EquivalentSnapshotActor = GetRawPointerOrNull(UE::LevelSnapshots::Private::GetPreallocatedIfCached(OriginalOwningActor, Cache));
 		const bool bAreOwnedByEquivalentActors = EquivalentSnapshotActor == SnapshotOwningActor; 
 		if (!bAreOwnedByEquivalentActors)
 		{
@@ -489,5 +489,5 @@ bool UE::LevelSnapshots::Private::AreActorsEquivalent(UObject* SnapshotPropertyV
 
 	// The snapshot actor was already allocated, if some other snapshot actor is referencing it
 	const TOptional<TNonNullPtr<AActor>> PreallocatedSnapshotVersion = UE::LevelSnapshots::Private::GetPreallocatedIfCached(OriginalActorReference, Cache);
-	return PreallocatedSnapshotVersion.Get(nullptr) == SnapshotPropertyValue;
+	return GetRawPointerOrNull(PreallocatedSnapshotVersion) == SnapshotPropertyValue;
 }
\ No newline at end of file
diff --git a/Engine/Plugins/VirtualProduction/LevelSnapshots/Source/LevelSnapshots/Private/Filtering/SnapshotSelectionBuilder.cpp b/Engine/Plugins/VirtualProduction/LevelSnapshots/Source/LevelSnapshots/Private/Filtering/SnapshotSelectionBuilder.cpp
--- a/Engine/Plugins/VirtualProduction/LevelSnapshots/Source/LevelSnapshots/Private/Filtering/SnapshotSelectionBuilder.cpp
+++ b/Engine/Plugins/VirtualProduction/LevelSnapshots/Source/LevelSnapshots/Private/Filtering/SnapshotSelectionBuilder.cpp
@@ -62,19 +62,19 @@ namespace UE::LevelSnapshots::Private
 			{
 				const FSoftObjectPath ActorPath = WorldActor;
 				const TOptional<TNonNullPtr<AActor>> DeserializedSnapshotActor = Snapshot->GetDeserializedActor(ActorPath);
-				if (!ensureAlwaysMsgf(DeserializedSnapshotActor.Get(nullptr), TEXT("Failed to get TMap value for key %s. Is the snapshot corrupted?"), *ActorPath.ToString()))
+				if (!ensureAlwaysMsgf(GetRawPointerOrNull(DeserializedSnapshotActor), TEXT("Failed to get TMap value for key %s. Is the snapshot corrupted?"), *ActorPath.ToString()))
 				{
 					continue;
 				}
 		
-				const EFilterResult::Type ActorInclusionResult = Filter->IsActorValid(FIsActorValidParams(DeserializedSnapshotActor.GetValue(), WorldActor));
+				const EFilterResult::Type ActorInclusionResult = Filter->IsActorValid(FIsActorValidParams(DeserializedSnapshotActor->Get(), WorldActor));
 				if (EFilterResult::CanInclude(ActorInclusionResult))
 				{
 					ULevelSnapshotsFilteringLibrary::ApplyFilterToFindSelectedProperties(
 						Snapshot,
 						Result,
 						WorldActor,
-						DeserializedSnapshotActor.GetValue(),
+						DeserializedSnapshotActor->Get(),
 						Filter
 					);
 				}
@@ -91,7 +91,7 @@ namespace UE::LevelSnapshots::Private
 					*RemovedActorsIt,
 					[this](const FSoftObjectPath& ObjectPath)
 					{
-						return Snapshot->GetDeserializedActor(ObjectPath).Get(nullptr);
+						return GetRawPointerOrNull(Snapshot->GetDeserializedActor(ObjectPath));
 					}
 				)
 			);
diff --git a/Engine/Plugins/VirtualProduction/LevelSnapshots/Source/LevelSnapshots/Private/Restoration/Actor/ActorUtil_TransientWorld.cpp b/Engine/Plugins/VirtualProduction/LevelSnapshots/Source/LevelSnapshots/Private/Restoration/Actor/ActorUtil_TransientWorld.cpp
--- a/Engine/Plugins/VirtualProduction/LevelSnapshots/Source/LevelSnapshots/Private/Restoration/Actor/ActorUtil_TransientWorld.cpp
+++ b/Engine/Plugins/VirtualProduction/LevelSnapshots/Source/LevelSnapshots/Private/Restoration/Actor/ActorUtil_TransientWorld.cpp
@@ -313,7 +313,7 @@ TOptional<TNonNullPtr<AActor>> UE::LevelSnapshots::Private::GetDeserializedActor
 	const TOptional<TNonNullPtr<AActor>> Result = Internal::GetDeserialized(OriginalObjectPath, WorldData, Cache, LocalisationSnapshotPackage, bWarnIfNotFound);
 	// Look up cache again because GetDeserialized might have reallocated Cache.ActorCache hence CacheBeforeDeserialize is no longer valid
 	const FActorSnapshotCache& CacheAfterDeserialize = Cache.ActorCache.FindOrAdd(OriginalObjectPath);
-	Internal::ConditionallyRerunConstructionScript(Result.Get(nullptr), CacheAfterDeserialize.ObjectDependencies, WorldData, Cache, LocalisationSnapshotPackage);
+	Internal::ConditionallyRerunConstructionScript(GetRawPointerOrNull(Result), CacheAfterDeserialize.ObjectDependencies, WorldData, Cache, LocalisationSnapshotPackage);
 	return Result;
 }
 
diff --git a/Engine/Plugins/VirtualProduction/LevelSnapshots/Source/LevelSnapshots/Private/Restoration/Component/BaseComponentRestorer.h b/Engine/Plugins/VirtualProduction/LevelSnapshots/Source/LevelSnapshots/Private/Restoration/Component/BaseComponentRestorer.h
--- a/Engine/Plugins/VirtualProduction/LevelSnapshots/Source/LevelSnapshots/Private/Restoration/Component/BaseComponentRestorer.h
+++ b/Engine/Plugins/VirtualProduction/LevelSnapshots/Source/LevelSnapshots/Private/Restoration/Component/BaseComponentRestorer.h
@@ -76,7 +76,7 @@ namespace UE::LevelSnapshots::Private
 				bool bIsOwnedByComponent;
 				const FSoftObjectPath& ComponentOuterPath = WorldData.SerializedObjectReferences[SubobjectData.OuterIndex];
 				const TOptional<TNonNullPtr<FActorSnapshotData>> ComponentOuterData = UE::LevelSnapshots::Private::FindSavedActorDataUsingObjectPath(WorldData.ActorData, ComponentOuterPath, bIsOwnedByComponent);
-				const bool bIsOwnedByOtherActor = ComponentOuterData.Get(nullptr) != &SnapshotData;
+				const bool bIsOwnedByOtherActor = GetRawPointerOrNull(ComponentOuterData) != &SnapshotData;
 				if (!ensureAlwaysMsgf(ComponentOuterData, TEXT("Failed to recreate component %s because its outer %s did not have any associated actor data. Investigate."), *ComponentPath.ToString(), *ComponentOuterPath.ToString())
 					|| !ensureMsgf(!bIsOwnedByOtherActor, TEXT("Failed to recreate component %s because saved data indicates it is owned by another actor %s. Components normally are owned by the actor they're attached to. Investigate."), *ComponentPath.ToString(), *ComponentOuterPath.ToString()))
 				{
diff --git a/Engine/Plugins/VirtualProduction/LevelSnapshots/Source/LevelSnapshots/Private/SpecialSupport/ActorGroupRestoration.cpp b/Engine/Plugins/VirtualProduction/LevelSnapshots/Source/LevelSnapshots/Private/SpecialSupport/ActorGroupRestoration.cpp
--- a/Engine/Plugins/VirtualProduction/LevelSnapshots/Source/LevelSnapshots/Private/SpecialSupport/ActorGroupRestoration.cpp
+++ b/Engine/Plugins/VirtualProduction/LevelSnapshots/Source/LevelSnapshots/Private/SpecialSupport/ActorGroupRestoration.cpp
@@ -231,7 +231,7 @@ namespace UE::LevelSnapshots::Private::ActorGroupRestoration
 			{
 				if (const TOptional<TNonNullPtr<AActor>> SnapshotActor = Params.DeserializeFromSnapshotFunc())
 				{
-					const AGroupActor* AsGroupActor = Cast<AGroupActor>(SnapshotActor.GetValue());
+					const AGroupActor* AsGroupActor = Cast<AGroupActor>(SnapshotActor->Get());
 					check(AsGroupActor);
 					return Algo::AllOf(AsGroupActor->GroupActors, [](AActor* Actor){ return Actor != nullptr; })
 						? EFilterResult::DoNotCare : EFilterResult::Disallow;
diff --git a/Engine/Plugins/VirtualProduction/LevelSnapshots/Source/LevelSnapshots/Private/Tests/AssetUserDataTests.cpp b/Engine/Plugins/VirtualProduction/LevelSnapshots/Source/LevelSnapshots/Private/Tests/AssetUserDataTests.cpp
--- a/Engine/Plugins/VirtualProduction/LevelSnapshots/Source/LevelSnapshots/Private/Tests/AssetUserDataTests.cpp
+++ b/Engine/Plugins/VirtualProduction/LevelSnapshots/Source/LevelSnapshots/Private/Tests/AssetUserDataTests.cpp
@@ -236,7 +236,7 @@ namespace UE::LevelSnapshots::Private::Tests
 		{
 			UActorComponent* Component = Cast<UActorComponent>(MainWorldObject.GetObject());
 			TOptional<TNonNullPtr<AActor>> SnapshotActor = Snapshot->GetDeserializedActor(Component->GetOwner());
-			return SnapshotActor ? Cast<ASnapshotTestActor>(SnapshotActor.GetValue())->TestComponent : nullptr;
+			return SnapshotActor ? Cast<ASnapshotTestActor>(SnapshotActor->Get())->TestComponent : nullptr;
 		}
 	};
 	
diff --git a/Engine/Plugins/VirtualProduction/LevelSnapshots/Source/LevelSnapshots/Private/Tests/ComponentTests.cpp b/Engine/Plugins/VirtualProduction/LevelSnapshots/Source/LevelSnapshots/Private/Tests/ComponentTests.cpp
--- a/Engine/Plugins/VirtualProduction/LevelSnapshots/Source/LevelSnapshots/Private/Tests/ComponentTests.cpp
+++ b/Engine/Plugins/VirtualProduction/LevelSnapshots/Source/LevelSnapshots/Private/Tests/ComponentTests.cpp
@@ -280,7 +280,7 @@ namespace UE::LevelSnapshots::Private::Tests
 				}
 
 				USceneComponent* AsSceneComponent = Cast<USceneComponent>(*CounterpartComponent);
-				ASnapshotTestActor* CastedCounterpartActor = Cast<ASnapshotTestActor>(CounterpartActor.GetValue());
+				ASnapshotTestActor* CastedCounterpartActor = Cast<ASnapshotTestActor>(CounterpartActor->Get());
 				check(AsSceneComponent);
 				TestTrue(TEXT("Snapshot attach parent correct"), AsSceneComponent->GetAttachParent() == CastedCounterpartActor->GetMesh());
 			})
diff --git a/Engine/Plugins/VirtualProduction/LevelSnapshots/Source/LevelSnapshots/Private/Tests/SubobjectTests.cpp b/Engine/Plugins/VirtualProduction/LevelSnapshots/Source/LevelSnapshots/Private/Tests/SubobjectTests.cpp
--- a/Engine/Plugins/VirtualProduction/LevelSnapshots/Source/LevelSnapshots/Private/Tests/SubobjectTests.cpp
+++ b/Engine/Plugins/VirtualProduction/LevelSnapshots/Source/LevelSnapshots/Private/Tests/SubobjectTests.cpp
@@ -816,7 +816,7 @@ namespace UE::LevelSnapshots::Private::Tests
 					return;
 				}
 			
-				ASnapshotTestActor* ActorB_Snapshot = Cast<ASnapshotTestActor>(OptionalActorB_Snapshot.GetValue());
+				ASnapshotTestActor* ActorB_Snapshot = Cast<ASnapshotTestActor>(OptionalActorB_Snapshot->Get());
 				TestTrue(TEXT("ObjectArray contains null"), ActorB_Snapshot->ObjectArray.Num() == 1 && ActorB_Snapshot->ObjectArray[0] == nullptr);
 				TestTrue(TEXT("ObjectSet contains null"), ActorB_Snapshot->ObjectSet.Num() == 1 && ActorB_Snapshot->ObjectSet.begin().ElementIt->Value == nullptr);
 				TestTrue(TEXT("ObjectMap contains null"), ActorB_Snapshot->ObjectMap.Num() == 1 && ActorB_Snapshot->ObjectMap.Find("MapKey") && ActorB_Snapshot->ObjectMap["MapKey"] == nullptr);
diff --git a/Engine/Plugins/VirtualProduction/LevelSnapshots/Source/LevelSnapshotsEditor/Private/Data/FilterListData.cpp b/Engine/Plugins/VirtualProduction/LevelSnapshots/Source/LevelSnapshotsEditor/Private/Data/FilterListData.cpp
--- a/Engine/Plugins/VirtualProduction/LevelSnapshots/Source/LevelSnapshotsEditor/Private/Data/FilterListData.cpp
+++ b/Engine/Plugins/VirtualProduction/LevelSnapshots/Source/LevelSnapshotsEditor/Private/Data/FilterListData.cpp
@@ -30,7 +30,7 @@ TWeakObjectPtr<AActor> FFilterListData::GetSnapshotCounterpartFor(const AActor*
 {
 	const TOptional<TNonNullPtr<AActor>> DeserializedActor = RelatedSnapshot->GetDeserializedActor(WorldActor);
 	return ensureAlwaysMsgf(DeserializedActor, TEXT("Deserialized actor does no exist. Either the snapshots's container world was deleted or the snapshot has no counterpart for this actor"))
-		? DeserializedActor.GetValue() : nullptr;
+		? DeserializedActor->Get() : nullptr;
 }
 
 void FFilterListData::ForEachModifiedActor(TFunctionRef<void(AActor*)> Callback) const
