Fix a bunch of clang static analyzer warnings in Verse related modules.

#rnx

diff --git a/Engine/Plugins/Developer/Concert/ConcertApp/MultiUserServer/Source/MultiUserServer/Private/Widgets/Clients/Logging/Util/LogAckTracker.cpp b/Engine/Plugins/Developer/Concert/ConcertApp/MultiUserServer/Source/MultiUserServer/Private/Widgets/Clients/Logging/Util/LogAckTracker.cpp
--- a/Engine/Plugins/Developer/Concert/ConcertApp/MultiUserServer/Source/MultiUserServer/Private/Widgets/Clients/Logging/Util/LogAckTracker.cpp
+++ b/Engine/Plugins/Developer/Concert/ConcertApp/MultiUserServer/Source/MultiUserServer/Private/Widgets/Clients/Logging/Util/LogAckTracker.cpp
@@ -6,9 +6,9 @@
 
 #include "Widgets/Clients/Logging/Source/IConcertLogSource.h"
 
-FLogAckTracker::FLogAckTracker(TSharedRef<IConcertLogSource> LogSource, TSharedRef<IConcertServer> Server)
-	: LogSource(MoveTemp(LogSource))
-	, Server(MoveTemp(Server))
+FLogAckTracker::FLogAckTracker(TSharedRef<IConcertLogSource> InLogSource, TSharedRef<IConcertServer> InServer)
+	: LogSource(MoveTemp(InLogSource))
+	, Server(MoveTemp(InServer))
 {
 	LogSource->OnLogEntryAdded().AddRaw(this, &FLogAckTracker::OnLogEntryProduced);
 	Server->OnConcertMessageAcknowledgementReceived().AddRaw(this, &FLogAckTracker::OnAckProcessed);
diff --git a/Engine/Plugins/Developer/Concert/ConcertApp/MultiUserServer/Source/MultiUserServer/Private/Widgets/Clients/PackageTransmission/Model/FilteredPackageTransmissionModel.cpp b/Engine/Plugins/Developer/Concert/ConcertApp/MultiUserServer/Source/MultiUserServer/Private/Widgets/Clients/PackageTransmission/Model/FilteredPackageTransmissionModel.cpp
--- a/Engine/Plugins/Developer/Concert/ConcertApp/MultiUserServer/Source/MultiUserServer/Private/Widgets/Clients/PackageTransmission/Model/FilteredPackageTransmissionModel.cpp
+++ b/Engine/Plugins/Developer/Concert/ConcertApp/MultiUserServer/Source/MultiUserServer/Private/Widgets/Clients/PackageTransmission/Model/FilteredPackageTransmissionModel.cpp
@@ -5,9 +5,9 @@
 
 namespace UE::MultiUserServer
 {
-	FFilteredPackageTransmissionModel::FFilteredPackageTransmissionModel(TSharedRef<IPackageTransmissionEntrySource> RealSource, TSharedRef<IFilter<const FPackageTransmissionEntry&>> Filter)
-		: RealSource(MoveTemp(RealSource))
-		, Filter(MoveTemp(Filter))
+	FFilteredPackageTransmissionModel::FFilteredPackageTransmissionModel(TSharedRef<IPackageTransmissionEntrySource> InRealSource, TSharedRef<IFilter<const FPackageTransmissionEntry&>> InFilter)
+		: RealSource(MoveTemp(InRealSource))
+		, Filter(MoveTemp(InFilter))
 	{
 		RealSource->OnPackageEntriesAdded().AddRaw(this, &FFilteredPackageTransmissionModel::OnSourceEntriesAdded);
 		RealSource->OnPackageEntriesModified().AddRaw(this, &FFilteredPackageTransmissionModel::OnSourceEntriesModified);
diff --git a/Engine/Plugins/Developer/Concert/ConcertApp/MultiUserServer/Source/MultiUserServer/Private/Widgets/Clients/PackageTransmission/PackageTransmissionTabController.cpp b/Engine/Plugins/Developer/Concert/ConcertApp/MultiUserServer/Source/MultiUserServer/Private/Widgets/Clients/PackageTransmission/PackageTransmissionTabController.cpp
--- a/Engine/Plugins/Developer/Concert/ConcertApp/MultiUserServer/Source/MultiUserServer/Private/Widgets/Clients/PackageTransmission/PackageTransmissionTabController.cpp
+++ b/Engine/Plugins/Developer/Concert/ConcertApp/MultiUserServer/Source/MultiUserServer/Private/Widgets/Clients/PackageTransmission/PackageTransmissionTabController.cpp
@@ -14,26 +14,25 @@
 namespace UE::MultiUserServer
 {
 	FPackageTransmissionTabController::FPackageTransmissionTabController(
-			FName TabId,
-			TSharedRef<FTabManager> OwningTabManager,
-			TSharedRef<FWorkspaceItem> WorkspaceItem,
-			TSharedRef<FPackageTransmissionModel> TransmissionModel,
-			TSharedRef<FEndpointToUserNameCache> EndpointToUserNameCache,
-			FCanScrollToLog CanScrollToLogDelegate,
-			FScrollToLog ScrollToLogDelegate
+			FName InTabId,
+			TSharedRef<FTabManager> InOwningTabManager,
+			TSharedRef<FWorkspaceItem> InWorkspaceItem,
+			TSharedRef<FPackageTransmissionModel> InTransmissionModel,
+			TSharedRef<FEndpointToUserNameCache> InEndpointToUserNameCache,
+			FCanScrollToLog InCanScrollToLogDelegate,
+			FScrollToLog InScrollToLogDelegate
 			)
-		: TabId(TabId)
-		, OwningTabManager(MoveTemp(OwningTabManager))
-		, TransmissionModel(MoveTemp(TransmissionModel))
-		, EndpointToUserNameCache(MoveTemp(EndpointToUserNameCache))
-		, CanScrollToLogDelegate(MoveTemp(CanScrollToLogDelegate))
-		, ScrollToLogDelegate(MoveTemp(ScrollToLogDelegate))
-		, Tokenizer(MakeShared<FPackageTransmissionEntryTokenizer>(EndpointToUserNameCache))
+		: TabId(InTabId)
+		, OwningTabManager(MoveTemp(InOwningTabManager))
+		, TransmissionModel(MoveTemp(InTransmissionModel))
+		, CanScrollToLogDelegate(MoveTemp(InCanScrollToLogDelegate))
+		, ScrollToLogDelegate(MoveTemp(InScrollToLogDelegate))
+		, Tokenizer(MakeShared<FPackageTransmissionEntryTokenizer>(MoveTemp(InEndpointToUserNameCache)))
 	{
 		OwningTabManager->RegisterTabSpawner(TabId, FOnSpawnTab::CreateRaw(this, &FPackageTransmissionTabController::SpawnTab))
 			// In the future we may create multiple FPackageTransmissionTabController and may want to change this name to depend on some parameter
 			.SetDisplayName(LOCTEXT("PackageTabLabel", "Packages"))
-			.SetGroup(WorkspaceItem)
+			.SetGroup(InWorkspaceItem)
 			.SetIcon(FSlateIcon(FConcertServerStyle::GetStyleSetName(), TEXT("Concert.Icon.Package")));
 	}
 
diff --git a/Engine/Plugins/Developer/Concert/ConcertApp/MultiUserServer/Source/MultiUserServer/Private/Widgets/Clients/PackageTransmission/PackageTransmissionTabController.h b/Engine/Plugins/Developer/Concert/ConcertApp/MultiUserServer/Source/MultiUserServer/Private/Widgets/Clients/PackageTransmission/PackageTransmissionTabController.h
--- a/Engine/Plugins/Developer/Concert/ConcertApp/MultiUserServer/Source/MultiUserServer/Private/Widgets/Clients/PackageTransmission/PackageTransmissionTabController.h
+++ b/Engine/Plugins/Developer/Concert/ConcertApp/MultiUserServer/Source/MultiUserServer/Private/Widgets/Clients/PackageTransmission/PackageTransmissionTabController.h
@@ -37,7 +37,6 @@ namespace UE::MultiUserServer
 		FName TabId;
 		TSharedRef<FTabManager> OwningTabManager;
 		TSharedRef<FPackageTransmissionModel> TransmissionModel;
-		TSharedRef<FEndpointToUserNameCache> EndpointToUserNameCache;
 		
 		FCanScrollToLog CanScrollToLogDelegate;
 		FScrollToLog ScrollToLogDelegate;
diff --git a/Engine/Plugins/Developer/Concert/ConcertApp/MultiUserServer/Source/MultiUserServer/Private/Widgets/Clients/Util/EndpointToUserNameCache.cpp b/Engine/Plugins/Developer/Concert/ConcertApp/MultiUserServer/Source/MultiUserServer/Private/Widgets/Clients/Util/EndpointToUserNameCache.cpp
--- a/Engine/Plugins/Developer/Concert/ConcertApp/MultiUserServer/Source/MultiUserServer/Private/Widgets/Clients/Util/EndpointToUserNameCache.cpp
+++ b/Engine/Plugins/Developer/Concert/ConcertApp/MultiUserServer/Source/MultiUserServer/Private/Widgets/Clients/Util/EndpointToUserNameCache.cpp
@@ -12,8 +12,8 @@
 #include "Algo/AnyOf.h"
 #include "Features/IModularFeatures.h"
 
-FEndpointToUserNameCache::FEndpointToUserNameCache(TSharedRef<IConcertServer> Server)
-	: Server(MoveTemp(Server))
+FEndpointToUserNameCache::FEndpointToUserNameCache(TSharedRef<IConcertServer> InServer)
+	: Server(MoveTemp(InServer))
 {
 	ConcertServerEvents::OnLiveSessionCreated().AddRaw(this, &FEndpointToUserNameCache::OnLiveSessionCreated);
 	ConcertServerEvents::OnLiveSessionDestroyed().AddRaw(this, &FEndpointToUserNameCache::OnLiveSessionDestroyed);
diff --git a/Engine/Plugins/Developer/Concert/ConcertApp/MultiUserServer/Source/MultiUserServer/Private/Widgets/SessionTabs/Live/LiveConcertSessionTab.cpp b/Engine/Plugins/Developer/Concert/ConcertApp/MultiUserServer/Source/MultiUserServer/Private/Widgets/SessionTabs/Live/LiveConcertSessionTab.cpp
--- a/Engine/Plugins/Developer/Concert/ConcertApp/MultiUserServer/Source/MultiUserServer/Private/Widgets/SessionTabs/Live/LiveConcertSessionTab.cpp
+++ b/Engine/Plugins/Developer/Concert/ConcertApp/MultiUserServer/Source/MultiUserServer/Private/Widgets/SessionTabs/Live/LiveConcertSessionTab.cpp
@@ -10,13 +10,13 @@
 #include "Widgets/Docking/SDockTab.h"
 #include "Widgets/SConcertTabViewBase.h"
 
-FLiveConcertSessionTab::FLiveConcertSessionTab(TSharedRef<IConcertServerSession> InspectedSession, TSharedRef<IConcertSyncServer> SyncServer, TAttribute<TSharedRef<SWindow>> ConstructUnderWindow, FShowConnectedClients OnConnectedClientsClicked)
-	: FConcertSessionTabBase(InspectedSession->GetSessionInfo().SessionId, SyncServer)
-	, InspectedSession(MoveTemp(InspectedSession))
-	, ConstructUnderWindow(MoveTemp(ConstructUnderWindow))
-	, OnConnectedClientsClicked(MoveTemp(OnConnectedClientsClicked))
-	, SessionHistoryController(MakeShared<FLiveServerSessionHistoryController>(InspectedSession, SyncServer))
-	, PackageViewerController(MakeShared<FConcertSessionPackageViewerController>(InspectedSession, SyncServer))
+FLiveConcertSessionTab::FLiveConcertSessionTab(TSharedRef<IConcertServerSession> InInspectedSession, TSharedRef<IConcertSyncServer> InSyncServer, TAttribute<TSharedRef<SWindow>> InConstructUnderWindow, FShowConnectedClients InOnConnectedClientsClicked)
+	: FConcertSessionTabBase(InInspectedSession->GetSessionInfo().SessionId, InSyncServer)
+	, InspectedSession(MoveTemp(InInspectedSession))
+	, ConstructUnderWindow(MoveTemp(InConstructUnderWindow))
+	, OnConnectedClientsClicked(MoveTemp(InOnConnectedClientsClicked))
+	, SessionHistoryController(MakeShared<FLiveServerSessionHistoryController>(InspectedSession, InSyncServer))
+	, PackageViewerController(MakeShared<FConcertSessionPackageViewerController>(InspectedSession, InSyncServer))
 {}
 
 void FLiveConcertSessionTab::CreateDockContent(const TSharedRef<SDockTab>& InDockTab)
diff --git a/Engine/Plugins/Developer/Concert/ConcertApp/MultiUserServer/Source/MultiUserServer/Private/Widgets/SessionTabs/Live/LiveServerSessionHistoryController.cpp b/Engine/Plugins/Developer/Concert/ConcertApp/MultiUserServer/Source/MultiUserServer/Private/Widgets/SessionTabs/Live/LiveServerSessionHistoryController.cpp
--- a/Engine/Plugins/Developer/Concert/ConcertApp/MultiUserServer/Source/MultiUserServer/Private/Widgets/SessionTabs/Live/LiveServerSessionHistoryController.cpp
+++ b/Engine/Plugins/Developer/Concert/ConcertApp/MultiUserServer/Source/MultiUserServer/Private/Widgets/SessionTabs/Live/LiveServerSessionHistoryController.cpp
@@ -21,9 +21,9 @@ namespace UE::MultiUserServer::Private
 	}
 }
 
-FLiveServerSessionHistoryController::FLiveServerSessionHistoryController(TSharedRef<IConcertServerSession> InspectedSession, TSharedRef<IConcertSyncServer> SyncServer)
-	: FServerSessionHistoryControllerBase(InspectedSession->GetId(), UE::MultiUserServer::Private::MakeLiveTabSessionHistoryArguments())
-	, SyncServer(MoveTemp(SyncServer))
+FLiveServerSessionHistoryController::FLiveServerSessionHistoryController(TSharedRef<IConcertServerSession> InInspectedSession, TSharedRef<IConcertSyncServer> InSyncServer)
+	: FServerSessionHistoryControllerBase(InInspectedSession->GetId(), UE::MultiUserServer::Private::MakeLiveTabSessionHistoryArguments())
+	, SyncServer(MoveTemp(InSyncServer))
 {
 	ReloadActivities();
 	UMultiUserServerColumnVisibilitySettings::GetSettings()->OnLiveActivityBrowserColumnVisibility().AddRaw(this, &FLiveServerSessionHistoryController::OnActivityListColumnVisibilitySettingsUpdated);
diff --git a/Engine/Plugins/Developer/Concert/ConcertApp/MultiUserServer/Source/MultiUserServer/Private/Widgets/SessionTabs/Live/PackageViewer/ConcertSessionPackageViewerController.cpp b/Engine/Plugins/Developer/Concert/ConcertApp/MultiUserServer/Source/MultiUserServer/Private/Widgets/SessionTabs/Live/PackageViewer/ConcertSessionPackageViewerController.cpp
--- a/Engine/Plugins/Developer/Concert/ConcertApp/MultiUserServer/Source/MultiUserServer/Private/Widgets/SessionTabs/Live/PackageViewer/ConcertSessionPackageViewerController.cpp
+++ b/Engine/Plugins/Developer/Concert/ConcertApp/MultiUserServer/Source/MultiUserServer/Private/Widgets/SessionTabs/Live/PackageViewer/ConcertSessionPackageViewerController.cpp
@@ -10,9 +10,9 @@
 #include "SConcertSessionPackageViewer.h"
 #include "Templates/NonNullPointer.h"
 
-FConcertSessionPackageViewerController::FConcertSessionPackageViewerController(TSharedRef<IConcertServerSession> InspectedSession, TSharedRef<IConcertSyncServer> SyncServer)
-	: InspectedSession(MoveTemp(InspectedSession))
-	, SyncServer(MoveTemp(SyncServer))
+FConcertSessionPackageViewerController::FConcertSessionPackageViewerController(TSharedRef<IConcertServerSession> InInspectedSession, TSharedRef<IConcertSyncServer> InSyncServer)
+	: InspectedSession(MoveTemp(InInspectedSession))
+	, SyncServer(MoveTemp(InSyncServer))
 	, PackageViewer(MakePackageViewer())
 {
 	ReloadActivities();
diff --git a/Engine/Source/Programs/UnrealBuildTool/Platform/Windows/UEBuildWindows.cs b/Engine/Source/Programs/UnrealBuildTool/Platform/Windows/UEBuildWindows.cs
--- a/Engine/Source/Programs/UnrealBuildTool/Platform/Windows/UEBuildWindows.cs
+++ b/Engine/Source/Programs/UnrealBuildTool/Platform/Windows/UEBuildWindows.cs
@@ -1173,6 +1173,7 @@ namespace UnrealBuildTool
 						break;
 					case StaticAnalyzer.None:
 					case StaticAnalyzer.Clang:
+					case StaticAnalyzer.Default:
 						Target.WindowsPlatform.Compiler = WindowsCompiler.ClangRTFM;
 						break;
 				}
