Allow users to disable logging in the multi-user server

diff --git a/Engine/Plugins/Developer/Concert/ConcertApp/MultiUserServer/Source/MultiUserServer/Private/Widgets/Clients/Logging/SConcertTransportLogFooter.cpp b/Engine/Plugins/Developer/Concert/ConcertApp/MultiUserServer/Source/MultiUserServer/Private/Widgets/Clients/Logging/SConcertTransportLogFooter.cpp
--- a/Engine/Plugins/Developer/Concert/ConcertApp/MultiUserServer/Source/MultiUserServer/Private/Widgets/Clients/Logging/SConcertTransportLogFooter.cpp
+++ b/Engine/Plugins/Developer/Concert/ConcertApp/MultiUserServer/Source/MultiUserServer/Private/Widgets/Clients/Logging/SConcertTransportLogFooter.cpp
@@ -2,8 +2,11 @@
 
 #include "SConcertTransportLogFooter.h"
 
+#include "ConcertTransportEvents.h"
 #include "Framework/Views/TableViewMetadata.h"
 #include "Widgets/Clients/Logging/Filter/FilteredConcertLogList.h"
+
+#include "Widgets/Input/SButton.h"
 #include "Widgets/Input/SComboBox.h"
 #include "Widgets/Input/SSpinBox.h"
 #include "Widgets/Layout/SSpacer.h"
@@ -34,7 +37,7 @@ namespace UE::MultiUserServer
 		ChildSlot
 		[
 			SNew(SHorizontalBox)
-			
+
 			// Page selection
 			+SHorizontalBox::Slot()
 			.AutoWidth()
@@ -104,6 +107,18 @@ namespace UE::MultiUserServer
 			[
 				SNew(SSpacer)
 			]
+			// Disable logging button
+			+SHorizontalBox::Slot()
+			.AutoWidth()
+			[
+				SNew(SButton)
+				.OnClicked_Lambda([](){ ConcertTransportEvents::SetLoggingEnabled(false); return FReply::Handled(); })
+				.ToolTipText(LOCTEXT("DisableLoggingVisibility.Button.Tooltip", "Executes Concert.DisableLogging"))
+				[
+					SNew(STextBlock)
+					.Text(LOCTEXT("DisableLoggingVisibility.Button.Text", "Disable Logging"))
+				]
+			]
 		];
 	}
 }
diff --git a/Engine/Plugins/Developer/Concert/ConcertMain/Source/Concert/Private/ConcertLogger.cpp b/Engine/Plugins/Developer/Concert/ConcertMain/Source/Concert/Private/ConcertLogger.cpp
--- a/Engine/Plugins/Developer/Concert/ConcertMain/Source/Concert/Private/ConcertLogger.cpp
+++ b/Engine/Plugins/Developer/Concert/ConcertMain/Source/Concert/Private/ConcertLogger.cpp
@@ -274,7 +274,14 @@ namespace ConcertTransportEvents
 	
 	void SetLoggingEnabled(bool bEnabled)
 	{
-		return GetActiveLoggers().EnableLogging();
+		if (bEnabled)
+		{
+			 GetActiveLoggers().EnableLogging();
+		}
+		else
+		{
+			GetActiveLoggers().DisableLogging();
+		}
 	}
 }
 
