Initialize FDialogCommands on Multi-user slate server startup

diff --git a/Engine/Plugins/Developer/Concert/ConcertApp/MultiUserServer/Source/MultiUserServer/Private/MultiUserServerModule.cpp b/Engine/Plugins/Developer/Concert/ConcertApp/MultiUserServer/Source/MultiUserServer/Private/MultiUserServerModule.cpp
--- a/Engine/Plugins/Developer/Concert/ConcertApp/MultiUserServer/Source/MultiUserServer/Private/MultiUserServerModule.cpp
+++ b/Engine/Plugins/Developer/Concert/ConcertApp/MultiUserServer/Source/MultiUserServer/Private/MultiUserServerModule.cpp
@@ -14,6 +14,8 @@
 #include "StandaloneRenderer.h"
 #include "Misc/ConfigCacheIni.h"
 
+#include "Dialog/DialogCommands.h"
+
 #define LOCTEXT_NAMESPACE "UnrealMultiUserUI"
 
 namespace UE::MultiUserServer
@@ -21,6 +23,11 @@ namespace UE::MultiUserServer
     void FConcertServerUIModule::StartupModule() 
 	{
 		MultiUserServerLayoutIni = GConfig->GetConfigFilename(TEXT("MultiUserServerLayout"));
+
+		if (!FDialogCommands::IsRegistered())
+		{
+			FDialogCommands::Register();
+		}
 	}
 	
 	void FConcertServerUIModule::ShutdownModule()
@@ -77,4 +84,4 @@ namespace UE::MultiUserServer
 
 IMPLEMENT_MODULE(UE::MultiUserServer::FConcertServerUIModule, MultiUserServer);
 
-#undef LOCTEXT_NAMESPACE
\ No newline at end of file
+#undef LOCTEXT_NAMESPACE
