Add support for displaying the window size in our client UI view.  This helps us see how big of a send window nodes have.  The bigger the send window the better the performance of the nodes.

diff --git a/Engine/Plugins/Developer/Concert/ConcertApp/MultiUserServer/Source/MultiUserServer/Private/Widgets/Clients/Browser/Models/IClientNetworkStatisticsModel.cpp b/Engine/Plugins/Developer/Concert/ConcertApp/MultiUserServer/Source/MultiUserServer/Private/Widgets/Clients/Browser/Models/IClientNetworkStatisticsModel.cpp
--- a/Engine/Plugins/Developer/Concert/ConcertApp/MultiUserServer/Source/MultiUserServer/Private/Widgets/Clients/Browser/Models/IClientNetworkStatisticsModel.cpp
+++ b/Engine/Plugins/Developer/Concert/ConcertApp/MultiUserServer/Source/MultiUserServer/Private/Widgets/Clients/Browser/Models/IClientNetworkStatisticsModel.cpp
@@ -44,6 +44,11 @@ namespace UE::MultiUserServer::NetworkStatistics
 		return FString::Printf(TEXT("%lld %s"), Unit.Value, FUnitConversion::GetUnitDisplayString(Unit.Units));
 	}
 	
+	FString FormatWindowSize(const FMessageTransportStatistics& Stats)
+	{
+		return FString::Printf(TEXT("%lld"), Stats.WindowSize);
+	}
+	
 	FString FormatTotalBytesLost(const FMessageTransportStatistics& Stats)
 	{
 		const FNumericUnit<uint64> Unit = FUnitConversion::QuantizeUnitsToBestFit(Stats.TotalBytesLost, EUnit::Bytes);
diff --git a/Engine/Plugins/Developer/Concert/ConcertApp/MultiUserServer/Source/MultiUserServer/Private/Widgets/Clients/Browser/Models/IClientNetworkStatisticsModel.h b/Engine/Plugins/Developer/Concert/ConcertApp/MultiUserServer/Source/MultiUserServer/Private/Widgets/Clients/Browser/Models/IClientNetworkStatisticsModel.h
--- a/Engine/Plugins/Developer/Concert/ConcertApp/MultiUserServer/Source/MultiUserServer/Private/Widgets/Clients/Browser/Models/IClientNetworkStatisticsModel.h
+++ b/Engine/Plugins/Developer/Concert/ConcertApp/MultiUserServer/Source/MultiUserServer/Private/Widgets/Clients/Browser/Models/IClientNetworkStatisticsModel.h
@@ -35,6 +35,7 @@ namespace UE::MultiUserServer
 		FString FormatTotalBytesReceived(const FMessageTransportStatistics& Stats);
 		FString FormatAverageRTT(const FMessageTransportStatistics& Stats);
 		FString FormatBytesInflight(const FMessageTransportStatistics& Stats);
+		FString FormatWindowSize(const FMessageTransportStatistics& Stats);		
 		FString FormatTotalBytesLost(const FMessageTransportStatistics& Stats);
 	}
 }
diff --git a/Engine/Plugins/Developer/Concert/ConcertApp/MultiUserServer/Source/MultiUserServer/Private/Widgets/Clients/Browser/SClientNetworkStats.cpp b/Engine/Plugins/Developer/Concert/ConcertApp/MultiUserServer/Source/MultiUserServer/Private/Widgets/Clients/Browser/SClientNetworkStats.cpp
--- a/Engine/Plugins/Developer/Concert/ConcertApp/MultiUserServer/Source/MultiUserServer/Private/Widgets/Clients/Browser/SClientNetworkStats.cpp
+++ b/Engine/Plugins/Developer/Concert/ConcertApp/MultiUserServer/Source/MultiUserServer/Private/Widgets/Clients/Browser/SClientNetworkStats.cpp
@@ -40,6 +40,7 @@ namespace UE::MultiUserServer
 		ReceiveText->SetText(FText::FromString(NetworkStatistics::FormatTotalBytesReceived(Statistics)));
 		RoundTripTimeText->SetText(FText::FromString(NetworkStatistics::FormatAverageRTT(Statistics)));
 		InflightText->SetText(FText::FromString(NetworkStatistics::FormatBytesInflight(Statistics)));
+		WindowSize->SetText(FText::FromString(NetworkStatistics::FormatWindowSize(Statistics)));
 		LossText->SetText(FText::FromString(NetworkStatistics::FormatTotalBytesLost(Statistics)));
 	}
 
@@ -56,7 +57,7 @@ namespace UE::MultiUserServer
 		AddStatistic(Content, LOCTEXT("ReceiveLabel", "Received"), LOCTEXT("ReceiveTooltip", "Total bytes received from this client"), ReceiveText);
 		AddStatistic(Content, LOCTEXT("RttLabel", "RTT"), LOCTEXT("RttTooltip", "Round trip time"), RoundTripTimeText);
 		AddStatistic(Content, LOCTEXT("InflightLabel", "Inflight"), LOCTEXT("InflightTooltip", "Total reliable bytes awaiting an ack from client"), InflightText);
-
+		AddStatistic(Content, LOCTEXT("WindowLabel", "WSize"), LOCTEXT("WindowSizeTooltip", "The window size for the end point. This is how big of a window we utilize for sending data. Bigger values are better. Small values indicate poor performance."), WindowSize);
 		// Loss should be at the right
 		SHorizontalBox::FScopedWidgetSlotArguments LossSlat = Content->AddSlot();
 		AddStatistic(LossSlat, LOCTEXT("LossLabel", "Lost"), LOCTEXT("LostTooltip", "Total bytes lost while sending to the client"), LossText);
diff --git a/Engine/Plugins/Developer/Concert/ConcertApp/MultiUserServer/Source/MultiUserServer/Private/Widgets/Clients/Browser/SClientNetworkStats.h b/Engine/Plugins/Developer/Concert/ConcertApp/MultiUserServer/Source/MultiUserServer/Private/Widgets/Clients/Browser/SClientNetworkStats.h
--- a/Engine/Plugins/Developer/Concert/ConcertApp/MultiUserServer/Source/MultiUserServer/Private/Widgets/Clients/Browser/SClientNetworkStats.h
+++ b/Engine/Plugins/Developer/Concert/ConcertApp/MultiUserServer/Source/MultiUserServer/Private/Widgets/Clients/Browser/SClientNetworkStats.h
@@ -40,6 +40,7 @@ namespace UE::MultiUserServer
 		TSharedPtr<STextBlock> RoundTripTimeText;
 		TSharedPtr<STextBlock> InflightText;
 		TSharedPtr<STextBlock> LossText;
+		TSharedPtr<STextBlock> WindowSize;
 		
 		void UpdateStatistics(const FMessageTransportStatistics& Statistics) const;
 
