Multi User Slate Server: Fix crash when closing server while output log is docked.

diff --git a/Engine/Plugins/Developer/Concert/ConcertApp/MultiUserServer/Source/MultiUserServer/Private/Window/ConcertServerWindowController.cpp b/Engine/Plugins/Developer/Concert/ConcertApp/MultiUserServer/Source/MultiUserServer/Private/Window/ConcertServerWindowController.cpp
--- a/Engine/Plugins/Developer/Concert/ConcertApp/MultiUserServer/Source/MultiUserServer/Private/Window/ConcertServerWindowController.cpp
+++ b/Engine/Plugins/Developer/Concert/ConcertApp/MultiUserServer/Source/MultiUserServer/Private/Window/ConcertServerWindowController.cpp
@@ -17,7 +17,7 @@
 #include "Framework/Notifications/NotificationManager.h"
 #include "HAL/PlatformApplicationMisc.h"
 #include "Misc/ConfigCacheIni.h"
-
+#include "Misc/CoreDelegates.h"
 
 #define LOCTEXT_NAMESPACE "UnrealMultiUserUI.FConcertServerWindowController"
 
@@ -32,11 +32,13 @@ namespace UE::MultiUserServer
 	{
 		ConcertComponents.Add(SessionBrowserController);
 		ConcertComponents.Add(ClientsController);
+
+		FCoreDelegates::OnPreExit.AddRaw(this, &FConcertServerWindowController::SaveTabManagerLayout);
 	}
 
 	FConcertServerWindowController::~FConcertServerWindowController()
 	{
-		FGlobalTabmanager::Get()->SaveAllVisualState();
+		FCoreDelegates::OnPreExit.RemoveAll(this);
 		UnregisterFromSessionDestructionEvents();
 	}
 
@@ -182,6 +184,11 @@ namespace UE::MultiUserServer
 		    GConfig->Flush(false, MultiUserServerLayoutIni);
 		}
 	}
+
+	void FConcertServerWindowController::SaveTabManagerLayout()
+	{
+		FGlobalTabmanager::Get()->SaveAllVisualState();
+	}
 }
 
 #undef LOCTEXT_NAMESPACE 
diff --git a/Engine/Plugins/Developer/Concert/ConcertApp/MultiUserServer/Source/MultiUserServer/Private/Window/ConcertServerWindowController.h b/Engine/Plugins/Developer/Concert/ConcertApp/MultiUserServer/Source/MultiUserServer/Private/Window/ConcertServerWindowController.h
--- a/Engine/Plugins/Developer/Concert/ConcertApp/MultiUserServer/Source/MultiUserServer/Private/Window/ConcertServerWindowController.h
+++ b/Engine/Plugins/Developer/Concert/ConcertApp/MultiUserServer/Source/MultiUserServer/Private/Window/ConcertServerWindowController.h
@@ -86,5 +86,8 @@ namespace UE::MultiUserServer
 		
 		void OnWindowClosed(const TSharedRef<SWindow>& Window);
 		void SaveLayout() const;
+
+		/** Saves the layout of the global tab manager (on shutdown). */
+		void SaveTabManagerLayout();
 	};
 }
